Copying //UE4/Dev-AnimPhys to //UE4/Dev-Main (Source: //UE4/Dev-AnimPhys @ 3624379)

diff --git a/Engine/Source/Developer/MeshMergeUtilities/Private/MeshMergeDataTracker.cpp b/Engine/Source/Developer/MeshMergeUtilities/Private/MeshMergeDataTracker.cpp
--- a/Engine/Source/Developer/MeshMergeUtilities/Private/MeshMergeDataTracker.cpp
+++ b/Engine/Source/Developer/MeshMergeUtilities/Private/MeshMergeDataTracker.cpp
@@ -4,7 +4,7 @@
 #include "MeshMergeHelpers.h"
 
 FMeshMergeDataTracker::FMeshMergeDataTracker()
-	: AvailableLightMapUVChannel(INDEX_NONE)
+	: AvailableLightMapUVChannel(INDEX_NONE), SummedLightMapPixels(0)
 {
 	FMemory::Memzero(bWithVertexColors);
 	FMemory::Memzero(bOcuppiedUVChannels);
@@ -96,6 +96,16 @@ TConstLODIndexIterator FMeshMergeDataTracker::GetLODIndexIterator() const
 	return LODIndices.CreateConstIterator();
 }
 
+void FMeshMergeDataTracker::AddLightMapPixels(int32 Dimension)
+{
+	SummedLightMapPixels += FMath::Max(Dimension, 0);
+}
+
+int32 FMeshMergeDataTracker::GetLightMapDimension() const
+{
+	return FMath::CeilToInt(FMath::Sqrt(SummedLightMapPixels));
+}
+
 bool FMeshMergeDataTracker::DoesLODContainVertexColors(int32 LODIndex) const
 {
 	checkf(FMath::IsWithinInclusive(LODIndex, 0, MAX_STATIC_MESH_LODS - 1), TEXT("Invalid LOD index"));
diff --git a/Engine/Source/Developer/MeshMergeUtilities/Private/MeshMergeDataTracker.h b/Engine/Source/Developer/MeshMergeUtilities/Private/MeshMergeDataTracker.h
--- a/Engine/Source/Developer/MeshMergeUtilities/Private/MeshMergeDataTracker.h
+++ b/Engine/Source/Developer/MeshMergeUtilities/Private/MeshMergeDataTracker.h
@@ -120,6 +120,11 @@ public:
 	/** Iterates over LOD indices for mesh */
 	TConstLODIndexIterator GetLODIndexIterator() const;
 
+	/** Add number of lightmap pixels used for one of the Meshes */
+	void AddLightMapPixels(int32 Pixels);
+	/** Returns the texture dimension required to distribute all of the lightmap pixels */
+	int32 GetLightMapDimension() const;
+
 	/** Returns whether or not any raw mesh entry contains vertex colors for the specified LOD index */
 	bool DoesLODContainVertexColors(int32 LODIndex) const;
 	/** Returns whether or not any raw mesh entry contains texture coordinates for the specified UV channel and LOD index */
@@ -147,6 +152,7 @@ protected:
 	bool bOcuppiedUVChannels[MAX_STATIC_MESH_LODS][MAX_MESH_TEXTURE_COORDS];
 	/** First available UV channel across all RawMesh entries */
 	int32 AvailableLightMapUVChannel;
+	int32 SummedLightMapPixels;
 
 	/** Remapping pairs for each mesh and LOD index combination */
 	TMultiMap<FMeshLODKey, SectionRemapPair> UniqueSectionIndexPerLOD;
diff --git a/Engine/Source/Developer/MeshMergeUtilities/Private/MeshMergeUtilities.cpp b/Engine/Source/Developer/MeshMergeUtilities/Private/MeshMergeUtilities.cpp
--- a/Engine/Source/Developer/MeshMergeUtilities/Private/MeshMergeUtilities.cpp
+++ b/Engine/Source/Developer/MeshMergeUtilities/Private/MeshMergeUtilities.cpp
@@ -223,7 +223,7 @@ void FMeshMergeUtilities::BakeMaterialsForComponent(TArray<TWeakObjectPtr<UObjec
 						bool bUsesVertexData;
 						Material->AnalyzeMaterialProperty(Entry.Property, NumTextureCoordinates, bUsesVertexData);
 
-						if (!Entry.bUseConstantValue)
+						if (!Entry.bUseConstantValue && Entry.Property != MP_MAX)
 						{
 							MaterialSettings.PropertySizes.Add(Entry.Property, Entry.bUseCustomSize ? Entry.CustomSize : MaterialOptions->TextureSize);
 						}
@@ -269,10 +269,7 @@ void FMeshMergeUtilities::BakeMaterialsForComponent(TArray<TWeakObjectPtr<UObjec
 				// Add all user defined properties for baking out
 				for (const FPropertyEntry& Entry : MaterialOptions->Properties)
 				{
-					int32 NumTextureCoordinates;
-					bool bUsesVertexData;
-					Material->AnalyzeMaterialProperty(Entry.Property, NumTextureCoordinates, bUsesVertexData);
-					if (!Entry.bUseConstantValue && Material->IsPropertyActive(Entry.Property))
+					if (!Entry.bUseConstantValue && Material->IsPropertyActive(Entry.Property) && Entry.Property != MP_MAX)
 					{
 						MaterialSettings.PropertySizes.Add(Entry.Property, Entry.bUseCustomSize ? Entry.CustomSize : MaterialOptions->TextureSize);
 					}
@@ -319,7 +316,7 @@ void FMeshMergeUtilities::BakeMaterialsForComponent(TArray<TWeakObjectPtr<UObjec
 	FIntPoint ConstantSize(1, 1);
 	for (const FPropertyEntry& Entry : MaterialOptions->Properties)
 	{
-		if (Entry.bUseConstantValue)
+		if (Entry.bUseConstantValue && Entry.Property != MP_MAX)
 		{
 			ConstantData.SetNum(1, false);
 			ConstantData[0] = FColor(Entry.ConstantValue * 255.0f, Entry.ConstantValue * 255.0f, Entry.ConstantValue * 255.0f);
@@ -495,7 +492,7 @@ void FMeshMergeUtilities::DetermineMaterialVertexDataUsage(TArray<bool>& InOutMa
 		for (const FPropertyEntry& Entry : MaterialOptions->Properties)
 		{
 			// Don't have to check a property if the result is going to be constant anyway
-			if (!Entry.bUseConstantValue)
+			if (!Entry.bUseConstantValue && Entry.Property != MP_MAX)
 			{
 				int32 NumTextureCoordinates;
 				bool bUsesVertexData;
@@ -571,62 +568,160 @@ UMaterialOptions* FMeshMergeUtilities::PopulateMaterialOptions(const FMaterialPr
 	const bool bCustomSizes = MaterialSettings.TextureSizingType == TextureSizingType_UseManualOverrideTextureSize;
 
 	FPropertyEntry Property;
-	Property.Property = MP_BaseColor;
-	Property.bUseConstantValue = false;
-	Property.bUseCustomSize = bCustomSizes;
-	Property.CustomSize = bCustomSizes ? MaterialSettings.DiffuseTextureSize : MaterialSettings.TextureSize;
+	PopulatePropertyEntry(MaterialSettings, MP_BaseColor, Property);
 	MaterialOptions->Properties.Add(Property);
 
-	Property.Property = MP_Specular;
-	Property.bUseCustomSize = bCustomSizes;
-	Property.CustomSize = bCustomSizes ? MaterialSettings.SpecularTextureSize : MaterialSettings.TextureSize;
+	PopulatePropertyEntry(MaterialSettings, MP_Specular, Property);
 	if (MaterialSettings.bSpecularMap)
 		MaterialOptions->Properties.Add(Property);
 
-	Property.Property = MP_Roughness;
-	Property.bUseCustomSize = bCustomSizes;
-	Property.CustomSize = bCustomSizes ? MaterialSettings.RoughnessTextureSize : MaterialSettings.TextureSize;
+	PopulatePropertyEntry(MaterialSettings, MP_Roughness, Property);
 	if (MaterialSettings.bRoughnessMap)
 		MaterialOptions->Properties.Add(Property);
 
-	Property.Property = MP_Metallic;
-	Property.bUseCustomSize = bCustomSizes;
-	Property.CustomSize = bCustomSizes ? MaterialSettings.MetallicTextureSize : MaterialSettings.TextureSize;
+	PopulatePropertyEntry(MaterialSettings, MP_Metallic, Property);
 	if (MaterialSettings.bMetallicMap)
 		MaterialOptions->Properties.Add(Property);
 
-	Property.Property = MP_Normal;
-	Property.bUseCustomSize = bCustomSizes;
-	Property.CustomSize = bCustomSizes ? MaterialSettings.NormalTextureSize : MaterialSettings.TextureSize;
+	PopulatePropertyEntry(MaterialSettings, MP_Normal, Property);
 	if (MaterialSettings.bNormalMap)
 		MaterialOptions->Properties.Add(Property);
 
-	Property.Property = MP_Opacity;
-	Property.bUseCustomSize = bCustomSizes;
-	Property.CustomSize = bCustomSizes ? MaterialSettings.OpacityTextureSize : MaterialSettings.TextureSize;
+	PopulatePropertyEntry(MaterialSettings, MP_Opacity, Property);
 	if (MaterialSettings.bOpacityMap)
 		MaterialOptions->Properties.Add(Property);
 
-	Property.Property = MP_OpacityMask;
-	Property.bUseCustomSize = bCustomSizes;
-	Property.CustomSize = bCustomSizes ? MaterialSettings.OpacityMaskTextureSize : MaterialSettings.TextureSize;
+	PopulatePropertyEntry(MaterialSettings, MP_OpacityMask, Property);
 	if (MaterialSettings.bOpacityMaskMap)
 		MaterialOptions->Properties.Add(Property);
 
-	Property.Property = MP_EmissiveColor;
-	Property.bUseCustomSize = bCustomSizes;
-	Property.CustomSize = bCustomSizes ? MaterialSettings.EmissiveTextureSize : MaterialSettings.TextureSize;
+	PopulatePropertyEntry(MaterialSettings, MP_EmissiveColor, Property);
 	if (MaterialSettings.bEmissiveMap)
 		MaterialOptions->Properties.Add(Property);
 
-	Property.Property = MP_AmbientOcclusion;
-	Property.bUseCustomSize = bCustomSizes;
-	Property.CustomSize = bCustomSizes ? MaterialSettings.AmbientOcclusionTextureSize : MaterialSettings.TextureSize;
+	PopulatePropertyEntry(MaterialSettings, MP_AmbientOcclusion, Property);
 	if (MaterialSettings.bAmbientOcclusionMap)
 		MaterialOptions->Properties.Add(Property);
+
 	return MaterialOptions;
 }
 
+void FMeshMergeUtilities::PopulatePropertyEntry(const FMaterialProxySettings& MaterialSettings, EMaterialProperty MaterialProperty, FPropertyEntry& InOutPropertyEntry) const
+{
+	InOutPropertyEntry.Property = MaterialProperty;
+	switch (MaterialSettings.TextureSizingType)
+	{	
+		/** Set property output size to unique per-property user set sizes */
+		case TextureSizingType_UseManualOverrideTextureSize:
+		{
+			InOutPropertyEntry.bUseCustomSize = true;
+			InOutPropertyEntry.CustomSize = [MaterialSettings, MaterialProperty]() -> FIntPoint
+			{
+				switch (MaterialProperty)
+				{
+					case MP_BaseColor: return MaterialSettings.DiffuseTextureSize;
+					case MP_Specular: return MaterialSettings.SpecularTextureSize;
+					case MP_Roughness: return MaterialSettings.RoughnessTextureSize;
+					case MP_Metallic: return MaterialSettings.MetallicTextureSize;
+					case MP_Normal: return MaterialSettings.NormalTextureSize;
+					case MP_Opacity: return MaterialSettings.OpacityTextureSize;
+					case MP_OpacityMask: return MaterialSettings.OpacityMaskTextureSize;
+					case MP_EmissiveColor: return MaterialSettings.EmissiveTextureSize;
+					case MP_AmbientOcclusion: return MaterialSettings.AmbientOcclusionTextureSize;
+					default:
+					{
+						checkf(false, TEXT("Invalid Material Property"));
+						return FIntPoint();
+					}	
+				}
+			}();
+
+			break;
+		}
+		/** Set property output size to biased values off the TextureSize value (Normal at fullres, Diffuse at halfres, and anything else at quarter res */
+		case TextureSizingType_UseAutomaticBiasedSizes:
+		{
+			const FIntPoint FullRes = MaterialSettings.TextureSize;
+			const FIntPoint HalfRes = FIntPoint(FMath::Max(8, FullRes.X >> 1), FMath::Max(8, FullRes.Y >> 1));
+			const FIntPoint QuarterRes = FIntPoint(FMath::Max(4, FullRes.X >> 2), FMath::Max(4, FullRes.Y >> 2));
+
+			InOutPropertyEntry.bUseCustomSize = true;
+			InOutPropertyEntry.CustomSize = [FullRes, HalfRes, QuarterRes, MaterialSettings, MaterialProperty]() -> FIntPoint
+			{
+				switch (MaterialProperty)
+				{
+				case MP_Normal: return FullRes;
+				case MP_BaseColor: return HalfRes;
+				case MP_Specular: return QuarterRes;
+				case MP_Roughness: return QuarterRes;
+				case MP_Metallic: return QuarterRes;				
+				case MP_Opacity: return QuarterRes;
+				case MP_OpacityMask: return QuarterRes;
+				case MP_EmissiveColor: return QuarterRes;
+				case MP_AmbientOcclusion: return QuarterRes;
+				default:
+				{
+					checkf(false, TEXT("Invalid Material Property"));
+					return FIntPoint();
+				}
+				}
+			}();
+
+			break;
+		}
+ 		/** Set all sizes to TextureSize */
+		case TextureSizingType_UseSingleTextureSize:
+		case TextureSizingType_UseSimplygonAutomaticSizing:
+		{
+			InOutPropertyEntry.bUseCustomSize = false;
+			InOutPropertyEntry.CustomSize = MaterialSettings.TextureSize;
+			break;
+		}
+	}
+	/** Check whether or not a constant value should be used for this property */
+	InOutPropertyEntry.bUseConstantValue = [MaterialSettings, MaterialProperty]() -> bool
+	{
+		switch (MaterialProperty)
+		{
+			case MP_BaseColor: return false;
+			case MP_Normal: return !MaterialSettings.bNormalMap;
+			case MP_Specular: return !MaterialSettings.bSpecularMap;
+			case MP_Roughness: return !MaterialSettings.bRoughnessMap;
+			case MP_Metallic: return !MaterialSettings.bMetallicMap;
+			case MP_Opacity: return !MaterialSettings.bOpacityMap;
+			case MP_OpacityMask: return !MaterialSettings.bOpacityMaskMap;
+			case MP_EmissiveColor: return !MaterialSettings.bEmissiveMap;
+			case MP_AmbientOcclusion: return !MaterialSettings.bAmbientOcclusionMap;
+			default:
+			{
+				checkf(false, TEXT("Invalid Material Property"));
+				return false;
+			}
+		}
+	}();
+	/** Set the value if a constant value should be used for this property */
+	InOutPropertyEntry.ConstantValue = [MaterialSettings, MaterialProperty]() -> float
+	{
+		switch (MaterialProperty)
+		{
+			case MP_BaseColor: return 1.0f;
+			case MP_Normal: return 1.0f;
+			case MP_Specular: return MaterialSettings.SpecularConstant;
+			case MP_Roughness: return MaterialSettings.RoughnessConstant;
+			case MP_Metallic: return MaterialSettings.MetallicConstant;
+			case MP_Opacity: return MaterialSettings.OpacityConstant;
+			case MP_OpacityMask: return MaterialSettings.OpacityMaskConstant;
+			case MP_EmissiveColor: return 0.0f;
+			case MP_AmbientOcclusion: return MaterialSettings.AmbientOcclusionConstant;
+			default:
+			{
+				checkf(false, TEXT("Invalid Material Property"));
+				return 1.0f;
+			}
+		}
+	}();
+}
+
 void FMeshMergeUtilities::CopyTextureRect(const FColor* Src, const FIntPoint& SrcSize, FColor* Dst, const FIntPoint& DstSize, const FIntPoint& DstPos) const
 {
 	const int32 RowLength = SrcSize.X * sizeof(FColor);
@@ -976,10 +1071,10 @@ void FMeshMergeUtilities::CreateProxyMesh(const TArray<AActor*>& InActors, const
 	// Unique set of sections in mesh
 	TArray<FSectionInfo> UniqueSections;
 	TArray<FSectionInfo> Sections;
-
-
 	TMultiMap<uint32, uint32> SectionToMesh;
 
+	int32 SummedLightmapPixels = 0;
+
 	for (const UStaticMeshComponent* StaticMeshComponent : ComponentsToMerge)
 	{
 		const int32 ScreenSizeBasedLODLevel = Utilities->GetLODLevelForScreenSize(StaticMeshComponent, Utilities->CalculateScreenSizeFromDrawDistance(StaticMeshComponent->Bounds.SphereRadius, ProjectionMatrix, EstimatedDistance));
@@ -1004,6 +1099,11 @@ void FMeshMergeUtilities::CreateProxyMesh(const TArray<AActor*>& InActors, const
 
 			SectionToMesh.Add(UniqueIndex, MeshIndex);
 		}
+
+		int32 LightMapWidth, LightMapHeight;
+		StaticMeshComponent->GetLightMapResolution(LightMapWidth, LightMapHeight);
+		// Make sure we at least have some lightmap space allocated in case the static mesh is set up with invalid input
+		SummedLightmapPixels += FMath::Max(16, LightMapHeight * LightMapWidth);
 	}
 
 	TArray<UMaterialInterface*> UniqueMaterials;
@@ -1020,9 +1120,12 @@ void FMeshMergeUtilities::CreateProxyMesh(const TArray<AActor*>& InActors, const
 
 	UMaterialOptions* Options = PopulateMaterialOptions(InMeshProxySettings.MaterialSettings);
 	TArray<EMaterialProperty> MaterialProperties;
-	for (const FPropertyEntry& Entry : Options->Properties)
+	for (const FPropertyEntry& Entry : Options->Properties )
 	{
-		MaterialProperties.Add(Entry.Property);
+		if (Entry.Property != MP_MAX)
+		{
+			MaterialProperties.Add(Entry.Property);
+		}
 	}
 
 	// Mesh index / ( Mesh relative section index / output index )	
@@ -1044,7 +1147,7 @@ void FMeshMergeUtilities::CreateProxyMesh(const TArray<AActor*>& InActors, const
 
 		for (const FPropertyEntry& Entry : Options->Properties)
 		{
-			if (!Entry.bUseConstantValue && Material->IsPropertyActive(Entry.Property))
+			if (!Entry.bUseConstantValue && Material->IsPropertyActive(Entry.Property) && Entry.Property != MP_MAX)
 			{
 				MaterialSettings.PropertySizes.Add(Entry.Property, Entry.bUseCustomSize ? Entry.CustomSize : Options->TextureSize);
 			}
@@ -1168,7 +1271,7 @@ void FMeshMergeUtilities::CreateProxyMesh(const TArray<AActor*>& InActors, const
 	FIntPoint ConstantSize(1, 1);
 	for (const FPropertyEntry& Entry : Options->Properties)
 	{
-		if (Entry.bUseConstantValue)
+		if (Entry.bUseConstantValue && Entry.Property != MP_MAX)
 		{
 			ConstantData.SetNum(1, false);
 			ConstantData[0] = FColor(Entry.ConstantValue * 255.0f, Entry.ConstantValue * 255.0f, Entry.ConstantValue * 255.0f);
@@ -1226,6 +1329,12 @@ void FMeshMergeUtilities::CreateProxyMesh(const TArray<AActor*>& InActors, const
 	Data->ProxyBasePackageName = InProxyBasePackageName;
 	Data->CallbackDelegate = InProxyCreatedDelegate;
 
+	// Lightmap resolution
+	if (InMeshProxySettings.bComputeLightMapResolution)
+	{
+		Data->InProxySettings.LightMapResolution = FMath::CeilToInt(FMath::Sqrt(SummedLightmapPixels));
+	}
+
 	// Add this proxy job to map	
 	Processor->AddProxyJob(InGuid, Data);
 
@@ -1342,6 +1451,15 @@ void FMeshMergeUtilities::MergeComponentsToStaticMesh(const TArray<UPrimitiveCom
 			UStaticMeshComponent* Component = StaticMeshComponentsToMerge[ComponentIndex];
 			Adapters.Add(FStaticMeshComponentAdapter(Component));
 			FStaticMeshComponentAdapter& Adapter = Adapters.Last();
+			
+			if (InSettings.bComputedLightMapResolution)
+			{
+				int32 LightMapHeight, LightMapWidth;
+				if (Component->GetLightMapResolution(LightMapWidth, LightMapHeight))
+				{
+					DataTracker.AddLightMapPixels(LightMapWidth * LightMapHeight);
+				}
+			}			
 
 			const int32 NumLODs = Adapter.GetNumberOfLODs();
 			for (int32 LODIndex = 0; LODIndex < NumLODs; ++LODIndex)
@@ -1512,7 +1630,7 @@ void FMeshMergeUtilities::MergeComponentsToStaticMesh(const TArray<UPrimitiveCom
 		TMap<EMaterialProperty, FIntPoint> PropertySizes;
 		for (const FPropertyEntry& Entry : MaterialOptions->Properties)
 		{
-			if (!Entry.bUseConstantValue)
+			if (!Entry.bUseConstantValue && Entry.Property != MP_MAX)
 			{
 				PropertySizes.Add(Entry.Property, Entry.bUseCustomSize ? Entry.CustomSize : MaterialOptions->TextureSize);
 			}
@@ -1642,7 +1760,7 @@ void FMeshMergeUtilities::MergeComponentsToStaticMesh(const TArray<UPrimitiveCom
 		FIntPoint ConstantSize(1, 1);
 		for (const FPropertyEntry& Entry : MaterialOptions->Properties)
 		{
-			if (Entry.bUseConstantValue)
+			if (Entry.bUseConstantValue && Entry.Property != MP_MAX)
 			{
 				ConstantData.SetNum(1, false);
 				ConstantData[0] = FLinearColor(Entry.ConstantValue, Entry.ConstantValue, Entry.ConstantValue).ToFColor(true);
@@ -1666,8 +1784,11 @@ void FMeshMergeUtilities::MergeComponentsToStaticMesh(const TArray<UPrimitiveCom
 		FFlattenMaterial OutMaterial;
 		for (const FPropertyEntry& Entry : MaterialOptions->Properties)
 		{
-			EFlattenMaterialProperties OldProperty = NewToOldProperty(Entry.Property);
-			OutMaterial.SetPropertySize(OldProperty, Entry.bUseCustomSize ? Entry.CustomSize : MaterialOptions->TextureSize);
+			if (Entry.Property != MP_MAX)
+			{
+				EFlattenMaterialProperties OldProperty = NewToOldProperty(Entry.Property);
+				OutMaterial.SetPropertySize(OldProperty, Entry.bUseCustomSize ? Entry.CustomSize : MaterialOptions->TextureSize);
+			}
 		}
 
 		/** Reweighting */
@@ -2070,7 +2191,7 @@ void FMeshMergeUtilities::MergeComponentsToStaticMesh(const TArray<UPrimitiveCom
 				SrcModel->BuildSettings.bUseHighPrecisionTangentBasis = false;
 				SrcModel->BuildSettings.bUseFullPrecisionUVs = false;
 				SrcModel->BuildSettings.bGenerateLightmapUVs = InSettings.bGenerateLightMapUV;
-				SrcModel->BuildSettings.MinLightmapResolution = InSettings.TargetLightMapResolution;
+				SrcModel->BuildSettings.MinLightmapResolution = InSettings.bComputedLightMapResolution ? DataTracker.GetLightMapDimension() : InSettings.TargetLightMapResolution;
 				SrcModel->BuildSettings.SrcLightmapIndex = 0;
 				SrcModel->BuildSettings.DstLightmapIndex = LightMapUVChannel;
 
@@ -2113,6 +2234,7 @@ void FMeshMergeUtilities::MergeComponentsToStaticMesh(const TArray<UPrimitiveCom
 
 		//Set the Imported version before calling the build
 		StaticMesh->ImportVersion = EImportStaticMeshVersion::LastVersion;
+		StaticMesh->LightMapResolution = InSettings.bComputedLightMapResolution ? DataTracker.GetLightMapDimension() : InSettings.TargetLightMapResolution;
 
 		StaticMesh->Build(bSilent);
 		StaticMesh->PostEditChange();
diff --git a/Engine/Source/Developer/MeshMergeUtilities/Private/MeshMergeUtilities.h b/Engine/Source/Developer/MeshMergeUtilities/Private/MeshMergeUtilities.h
--- a/Engine/Source/Developer/MeshMergeUtilities/Private/MeshMergeUtilities.h
+++ b/Engine/Source/Developer/MeshMergeUtilities/Private/MeshMergeUtilities.h
@@ -3,6 +3,7 @@
 #pragma once
 
 #include "IMeshMergeUtilities.h"
+#include "SceneTypes.h"
 
 /**
  * Mesh Merge Utilities
@@ -37,6 +38,8 @@ protected:
 	float FlattenEmissivescale(TArray<struct FFlattenMaterial>& InMaterialList) const;
 	/** Populates material options object from legacy material proxy settings  */
 	UMaterialOptions* PopulateMaterialOptions(const FMaterialProxySettings& MaterialSettings) const;
+	/** Populates a single property entry with correct material baking settings */
+	void PopulatePropertyEntry(const FMaterialProxySettings& MaterialSettings, EMaterialProperty MaterialProperty, struct FPropertyEntry& InOutPropertyEntry) const;
 	/** Copies part (box) from a texture to another texture */
 	void CopyTextureRect(const FColor* Src, const FIntPoint& SrcSize, FColor* Dst, const FIntPoint& DstSize, const FIntPoint& DstPos) const;
 	/** Sets a part (box) on a texture to ColorValue */
diff --git a/Engine/Source/Developer/Windows/WindowsTargetPlatform/Classes/WindowsTargetSettings.h b/Engine/Source/Developer/Windows/WindowsTargetPlatform/Classes/WindowsTargetSettings.h
--- a/Engine/Source/Developer/Windows/WindowsTargetPlatform/Classes/WindowsTargetSettings.h
+++ b/Engine/Source/Developer/Windows/WindowsTargetPlatform/Classes/WindowsTargetSettings.h
@@ -78,4 +78,15 @@ public:
 	UPROPERTY(config, EditAnywhere, Category = "Audio", meta = (ClampMin = "0", UIMin = "0", DisplayName = "Number of Source Workers"))
 	int32 AudioNumSourceWorkers;
 
+	/** Which of the currently enabled spatialization plugins to use on Windows. */
+	UPROPERTY(config, EditAnywhere, Category = "Audio")
+	FString SpatializationPlugin;
+
+	/** Which of the currently enabled reverb plugins to use on Windows. */
+	UPROPERTY(config, EditAnywhere, Category = "Audio")
+	FString ReverbPlugin;
+
+	/** Which of the currently enabled occlusion plugins to use on Windows. */
+	UPROPERTY(config, EditAnywhere, Category = "Audio")
+	FString OcclusionPlugin;
 };
diff --git a/Engine/Source/Editor/AdvancedPreviewScene/Private/AdvancedPreviewScene.cpp b/Engine/Source/Editor/AdvancedPreviewScene/Private/AdvancedPreviewScene.cpp
--- a/Engine/Source/Editor/AdvancedPreviewScene/Private/AdvancedPreviewScene.cpp
+++ b/Engine/Source/Editor/AdvancedPreviewScene/Private/AdvancedPreviewScene.cpp
@@ -206,6 +206,12 @@ void FAdvancedPreviewScene::UpdateScene(FPreviewSceneProfile& Profile, bool bUpd
 	CurrentRotationSpeed = Profile.RotationSpeed;
 }
 
+FLinearColor FAdvancedPreviewScene::GetBackgroundColor() const
+{
+	FLinearColor Color = DefaultSettings->Profiles[CurrentProfileIndex].EnvironmentColor;
+	return Color * DefaultSettings->Profiles[CurrentProfileIndex].EnvironmentIntensity;
+}
+
 void FAdvancedPreviewScene::SetFloorOffset(const float InFloorOffset)
 {
 	FTransform FloorTransform(FRotator(0, 0, 0), FVector(0, 0, -(InFloorOffset)), FVector(4.0f, 4.0f, 1.0f));
diff --git a/Engine/Source/Editor/AdvancedPreviewScene/Public/AdvancedPreviewScene.h b/Engine/Source/Editor/AdvancedPreviewScene/Public/AdvancedPreviewScene.h
--- a/Engine/Source/Editor/AdvancedPreviewScene/Public/AdvancedPreviewScene.h
+++ b/Engine/Source/Editor/AdvancedPreviewScene/Public/AdvancedPreviewScene.h
@@ -30,6 +30,10 @@ public:
 
 	void UpdateScene(FPreviewSceneProfile& Profile, bool bUpdateSkyLight = true, bool bUpdateEnvironment = true, bool bUpdatePostProcessing = true, bool bUpdateDirectionalLight = true);
 
+	/** Begin FPreviewScene */
+	virtual FLinearColor GetBackgroundColor() const override;
+	/** End FPreviewScene */
+
 	/* Begin FTickableEditorObject */
 	virtual void Tick(float DeltaTime) override;
 	virtual bool IsTickable() const override;
diff --git a/Engine/Source/Editor/AdvancedPreviewScene/Public/AssetViewerSettings.h b/Engine/Source/Editor/AdvancedPreviewScene/Public/AssetViewerSettings.h
--- a/Engine/Source/Editor/AdvancedPreviewScene/Public/AssetViewerSettings.h
+++ b/Engine/Source/Editor/AdvancedPreviewScene/Public/AssetViewerSettings.h
@@ -32,6 +32,8 @@ struct FPreviewSceneProfile
 		SkyLightIntensity = 1.0f;
 		LightingRigRotation = 0.0f;
 		RotationSpeed = 2.0f;
+		EnvironmentIntensity = 1.0f;
+		EnvironmentColor = FLinearColor(0.2f, 0.2f, 0.2f, 1.0f);
 		// Set up default cube map texture from editor/engine textures
 		EnvironmentCubeMap = nullptr;
 		EnvironmentCubeMapPath = TEXT("/Engine/EditorMaterials/AssetViewer/EpicQuadPanorama_CC+EV1.EpicQuadPanorama_CC+EV1");
@@ -71,6 +73,14 @@ struct FPreviewSceneProfile
 	UPROPERTY(EditAnywhere, config, Category = Environment)
 	bool bShowFloor;
 
+	/** The environment color, used if Show Environment is false. */
+	UPROPERTY(EditAnywhere, config, Category = Environment, meta=(EditCondition="!bShowEnvironment"))
+	FLinearColor EnvironmentColor;
+	
+	/** The environment intensity (0.0 - 20.0), used if Show Environment is false. */
+	UPROPERTY(EditAnywhere, config, Category = Lighting, meta = (UIMin = "0.0", UIMax = "20.0", EditCondition="!bShowEnvironment"))
+	float EnvironmentIntensity;
+
 	/** Sets environment cube map used for sky lighting and reflections */
 	UPROPERTY(EditAnywhere, transient, Category = Environment)
 	TSoftObjectPtr<UTextureCube> EnvironmentCubeMap;
diff --git a/Engine/Source/Editor/AnimGraph/Classes/AnimGraphNode_StateMachineBase.h b/Engine/Source/Editor/AnimGraph/Classes/AnimGraphNode_StateMachineBase.h
--- a/Engine/Source/Editor/AnimGraph/Classes/AnimGraphNode_StateMachineBase.h
+++ b/Engine/Source/Editor/AnimGraph/Classes/AnimGraphNode_StateMachineBase.h
@@ -26,6 +26,7 @@ class ANIMGRAPH_API UAnimGraphNode_StateMachineBase : public UAnimGraphNode_Base
 	virtual FText GetNodeTitle(ENodeTitleType::Type TitleType) const override;
 	virtual void PostPlacedNewNode() override;
 	virtual UObject* GetJumpTargetForDoubleClick() const override;
+	virtual void JumpToDefinition() const override;
 	virtual void DestroyNode() override;
 	virtual void PostPasteNode() override;
 	virtual TSharedPtr<class INameValidatorInterface> MakeNameValidator() const override;
diff --git a/Engine/Source/Editor/AnimGraph/Classes/AnimPreviewInstance.h b/Engine/Source/Editor/AnimGraph/Classes/AnimPreviewInstance.h
--- a/Engine/Source/Editor/AnimGraph/Classes/AnimPreviewInstance.h
+++ b/Engine/Source/Editor/AnimGraph/Classes/AnimPreviewInstance.h
@@ -27,7 +27,7 @@ enum EMontagePreviewType
 
 /** Proxy override for this UAnimInstance-derived class */
 USTRUCT()
-struct FAnimPreviewInstanceProxy : public FAnimSingleNodeInstanceProxy
+struct ANIMGRAPH_API FAnimPreviewInstanceProxy : public FAnimSingleNodeInstanceProxy
 {
 	GENERATED_BODY()
 
diff --git a/Engine/Source/Editor/AnimGraph/Private/AnimGraphNode_ModifyCurve.cpp b/Engine/Source/Editor/AnimGraph/Private/AnimGraphNode_ModifyCurve.cpp
--- a/Engine/Source/Editor/AnimGraph/Private/AnimGraphNode_ModifyCurve.cpp
+++ b/Engine/Source/Editor/AnimGraph/Private/AnimGraphNode_ModifyCurve.cpp
@@ -41,21 +41,14 @@ TArray<FName> UAnimGraphNode_ModifyCurve::GetCurvesToAdd() const
 	const FSmartNameMapping* Mapping = GetAnimBlueprint()->TargetSkeleton->GetSmartNameContainer(USkeleton::AnimCurveMappingName);
 	if (Mapping)
 	{
-		TArray<SmartName::UID_Type> UidList;
-		Mapping->FillUidArray(UidList);
+		Mapping->FillNameArray(CurvesToAdd);
 
-		// Iterate through all curves..
-		for (SmartName::UID_Type Uid : UidList)
+		for (FName ExistingCurveName : Node.CurveNames)
 		{
-			FSmartName SmartName;
-			Mapping->FindSmartNameByUID(Uid, SmartName);
-
-			// Don't offer curves we already have
-			if (!Node.CurveNames.Contains(SmartName.DisplayName))
-			{
-				CurvesToAdd.Add(SmartName.DisplayName);
-			}
+			CurvesToAdd.RemoveSingleSwap(ExistingCurveName, false);
 		}
+
+		CurvesToAdd.Sort();
 	}
 
 	return CurvesToAdd;
diff --git a/Engine/Source/Editor/AnimGraph/Private/AnimGraphNode_SequencePlayer.cpp b/Engine/Source/Editor/AnimGraph/Private/AnimGraphNode_SequencePlayer.cpp
--- a/Engine/Source/Editor/AnimGraph/Private/AnimGraphNode_SequencePlayer.cpp
+++ b/Engine/Source/Editor/AnimGraph/Private/AnimGraphNode_SequencePlayer.cpp
@@ -75,7 +75,7 @@ FText UAnimGraphNode_SequencePlayer::GetTooltipText() const
 {
 	if (!Node.Sequence)
 	{
-		return GetTitleGivenAssetInfo(FText::FromString(TEXT("None")), false);
+		return FText();
 	}
 
 	const bool bAdditive = Node.Sequence->IsValidAdditive();
@@ -200,11 +200,16 @@ void UAnimGraphNode_SequencePlayer::GetMenuActions(FBlueprintActionDatabaseRegis
 				TWeakObjectPtr<UAnimSequence> AnimSequence = Cast<UAnimSequence>(Asset.GetAsset());
 				NodeSpawner->CustomizeNodeDelegate = UBlueprintNodeSpawner::FCustomizeNodeDelegate::CreateStatic(LoadedAssetSetup, AnimSequence);
 				NodeSpawner->DefaultMenuSignature.MenuName = GetTitleGivenAssetInfo(FText::FromName(AnimSequence->GetFName()), AnimSequence->IsValidAdditive());
+				NodeSpawner->DefaultMenuSignature.Tooltip = GetTitleGivenAssetInfo(FText::FromString(AnimSequence->GetPathName()), AnimSequence->IsValidAdditive());
 			}
 			else
 			{
+				const FString TagValue = Asset.GetTagValueRef<FString>(GET_MEMBER_NAME_CHECKED(UAnimSequence, AdditiveAnimType));
+				const bool bKnownToBeAdditive = (!TagValue.IsEmpty() && !TagValue.Equals(TEXT("AAT_None")));
+
 				NodeSpawner->CustomizeNodeDelegate = UBlueprintNodeSpawner::FCustomizeNodeDelegate::CreateStatic(UnloadedAssetSetup, Asset);
-				NodeSpawner->DefaultMenuSignature.MenuName = GetTitleGivenAssetInfo(FText::FromName(Asset.AssetName), /*bKnownToBeAdditive =*/false);
+				NodeSpawner->DefaultMenuSignature.MenuName = GetTitleGivenAssetInfo(FText::FromName(Asset.AssetName), bKnownToBeAdditive);
+				NodeSpawner->DefaultMenuSignature.Tooltip = GetTitleGivenAssetInfo(FText::FromName(Asset.ObjectPath), bKnownToBeAdditive);
 			}
 			ActionRegistrar.AddBlueprintAction(Asset, NodeSpawner);
 		}
@@ -216,6 +221,7 @@ void UAnimGraphNode_SequencePlayer::GetMenuActions(FBlueprintActionDatabaseRegis
 		TWeakObjectPtr<UAnimSequence> SequencePtr = AnimSequence;
 		NodeSpawner->CustomizeNodeDelegate = UBlueprintNodeSpawner::FCustomizeNodeDelegate::CreateStatic(LoadedAssetSetup, SequencePtr);
 		NodeSpawner->DefaultMenuSignature.MenuName = GetTitleGivenAssetInfo(FText::FromName(AnimSequence->GetFName()), AnimSequence->IsValidAdditive());
+		NodeSpawner->DefaultMenuSignature.Tooltip = GetTitleGivenAssetInfo(FText::FromString(AnimSequence->GetPathName()), AnimSequence->IsValidAdditive());
 
 		ActionRegistrar.AddBlueprintAction(QueryObject, NodeSpawner);
 	}
@@ -240,7 +246,12 @@ void UAnimGraphNode_SequencePlayer::GetMenuActions(FBlueprintActionDatabaseRegis
 
 			UBlueprintNodeSpawner* NodeSpawner = UBlueprintNodeSpawner::Create(GetClass());
 			NodeSpawner->CustomizeNodeDelegate = UBlueprintNodeSpawner::FCustomizeNodeDelegate::CreateStatic(UnloadedAssetSetup, Asset);
-			NodeSpawner->DefaultMenuSignature.MenuName = GetTitleGivenAssetInfo(FText::FromName(Asset.AssetName), /*bKnownToBeAdditive =*/false);
+
+			const FString TagValue = Asset.GetTagValueRef<FString>(GET_MEMBER_NAME_CHECKED(UAnimSequence, AdditiveAnimType));
+			const bool bKnownToBeAdditive = (!TagValue.IsEmpty() && !TagValue.Equals(TEXT("AAT_None")));
+
+			NodeSpawner->DefaultMenuSignature.MenuName = GetTitleGivenAssetInfo(FText::FromName(Asset.AssetName), bKnownToBeAdditive);
+			NodeSpawner->DefaultMenuSignature.Tooltip = GetTitleGivenAssetInfo(FText::FromName(Asset.ObjectPath), bKnownToBeAdditive);
 			ActionRegistrar.AddBlueprintAction(Asset, NodeSpawner);
 		}
 	}	
diff --git a/Engine/Source/Editor/AnimGraph/Private/AnimGraphNode_StateMachineBase.cpp b/Engine/Source/Editor/AnimGraph/Private/AnimGraphNode_StateMachineBase.cpp
--- a/Engine/Source/Editor/AnimGraph/Private/AnimGraphNode_StateMachineBase.cpp
+++ b/Engine/Source/Editor/AnimGraph/Private/AnimGraphNode_StateMachineBase.cpp
@@ -10,6 +10,7 @@
 #include "AnimationStateMachineGraph.h"
 #include "AnimationStateMachineSchema.h"
 #include "AnimGraphNode_StateMachine.h"
+#include "Kismet2/KismetEditorUtilities.h"
 
 /////////////////////////////////////////////////////
 // FAnimStateMachineNodeNameValidator
@@ -124,6 +125,14 @@ UObject* UAnimGraphNode_StateMachineBase::GetJumpTargetForDoubleClick() const
 	return EditorStateMachineGraph;
 }
 
+void UAnimGraphNode_StateMachineBase::JumpToDefinition() const
+{
+	if (UObject* HyperlinkTarget = GetJumpTargetForDoubleClick())
+	{
+		FKismetEditorUtilities::BringKismetToFocusAttentionOnObject(HyperlinkTarget);
+	}
+}
+
 void UAnimGraphNode_StateMachineBase::DestroyNode()
 {
 	UEdGraph* GraphToRemove = EditorStateMachineGraph;
diff --git a/Engine/Source/Editor/AnimGraph/Private/AnimGraphNode_SubInstance.cpp b/Engine/Source/Editor/AnimGraph/Private/AnimGraphNode_SubInstance.cpp
--- a/Engine/Source/Editor/AnimGraph/Private/AnimGraphNode_SubInstance.cpp
+++ b/Engine/Source/Editor/AnimGraph/Private/AnimGraphNode_SubInstance.cpp
@@ -12,6 +12,7 @@
 #include "Animation/AnimNode_SubInput.h"
 #include "PropertyCustomizationHelpers.h"
 #include "ScopedTransaction.h"
+#include "AnimationGraphSchema.h"
 
 #define LOCTEXT_NAMESPACE "SubInstanceNode"
 
@@ -115,6 +116,9 @@ void UAnimGraphNode_SubInstance::ReallocatePinsDuringReconstruction(TArray<UEdGr
 	// Need the schema to extract pin types
 	const UEdGraphSchema_K2* Schema = CastChecked<UEdGraphSchema_K2>(GetSchema());
 
+	// Default anim schema for util funcions
+	const UAnimationGraphSchema* AnimGraphDefaultSchema = GetDefault<UAnimationGraphSchema>();
+
 	bool bShowPose = false;
 
 	// Scan the target class for a sub input node, we only want to show the pose input if
@@ -198,6 +202,15 @@ void UAnimGraphNode_SubInstance::ReallocatePinsDuringReconstruction(TArray<UEdGr
 		ExposedPropertyNames.Remove(RemovedPropertyName);
 	}
 
+	// Make sure that any old pins that linked to properties are told not to be orphans
+	for(UEdGraphPin* OldPin : OldPins)
+	{
+		if(OldPin && !AnimGraphDefaultSchema->IsPosePin(OldPin->PinType))
+		{
+			OldPin->bSavePinIfOrphaned = false;
+		}
+	}
+
 }
 
 void UAnimGraphNode_SubInstance::PostEditChangeProperty(FPropertyChangedEvent& PropertyChangedEvent)
diff --git a/Engine/Source/Editor/AnimGraph/Private/PoseDriverDetails.cpp b/Engine/Source/Editor/AnimGraph/Private/PoseDriverDetails.cpp
--- a/Engine/Source/Editor/AnimGraph/Private/PoseDriverDetails.cpp
+++ b/Engine/Source/Editor/AnimGraph/Private/PoseDriverDetails.cpp
@@ -883,15 +883,13 @@ void FPoseDriverDetails::UpdateDrivenNameOptions()
 				const FSmartNameMapping* Mapping = Skeleton->GetSmartNameContainer(USkeleton::AnimCurveMappingName);
 				if (Mapping)
 				{
-					TArray<SmartName::UID_Type> UidList;
-					Mapping->FillUidArray(UidList);
+					TArray<FName> NameArray;
+					Mapping->FillNameArray(NameArray);
+					NameArray.Sort();
 
-					for (SmartName::UID_Type Uid : UidList)
+					for (FName CurveName : NameArray)
 					{
-						FSmartName SmartName;
-						Mapping->FindSmartNameByUID(Uid, SmartName);
-
-						DrivenNameOptions.Add(MakeShareable(new FName(SmartName.DisplayName)));
+						DrivenNameOptions.Add(MakeShareable(new FName(CurveName)));
 					}
 				}
 			}
diff --git a/Engine/Source/Editor/AnimationBlueprintEditor/Private/AnimationBlueprintEditor.cpp b/Engine/Source/Editor/AnimationBlueprintEditor/Private/AnimationBlueprintEditor.cpp
--- a/Engine/Source/Editor/AnimationBlueprintEditor/Private/AnimationBlueprintEditor.cpp
+++ b/Engine/Source/Editor/AnimationBlueprintEditor/Private/AnimationBlueprintEditor.cpp
@@ -61,6 +61,8 @@
 
 #include "AnimGraphNode_AimOffsetLookAt.h"
 #include "AnimGraphNode_RotationOffsetBlendSpace.h"
+#include "Algo/Transform.h"
+#include "ISkeletonTreeItem.h"
 
 #define LOCTEXT_NAMESPACE "AnimationBlueprintEditor"
 
@@ -202,13 +204,14 @@ void FAnimationBlueprintEditor::InitAnimationBlueprintEditor(const EToolkitMode:
 	PersonaToolkit = PersonaModule.CreatePersonaToolkit(InAnimBlueprint);
 
 	PersonaToolkit->GetPreviewScene()->SetDefaultAnimationMode(EPreviewSceneDefaultAnimationMode::AnimationBlueprint);
+	PersonaToolkit->GetPreviewScene()->RegisterOnPreviewMeshChanged(FOnPreviewMeshChanged::CreateSP(this, &FAnimationBlueprintEditor::HandlePreviewMeshChanged));
 
 	TSharedRef<IAssetFamily> AssetFamily = PersonaModule.CreatePersonaAssetFamily(InAnimBlueprint);
 	AssetFamily->RecordAssetOpened(FAssetData(InAnimBlueprint));
 
 	// create the skeleton tree
-	FSkeletonTreeArgs SkeletonTreeArgs(OnPostUndo);
-	SkeletonTreeArgs.OnObjectSelected = FOnObjectSelected::CreateSP(this, &FAnimationBlueprintEditor::HandleObjectSelected);
+	FSkeletonTreeArgs SkeletonTreeArgs;
+	SkeletonTreeArgs.OnSelectionChanged = FOnSkeletonTreeSelectionChanged::CreateSP(this, &FAnimationBlueprintEditor::HandleSelectionChanged);
 	SkeletonTreeArgs.PreviewScene = GetPreviewScene();
 
 	ISkeletonEditorModule& SkeletonEditorModule = FModuleManager::LoadModuleChecked<ISkeletonEditorModule>("SkeletonEditor");
@@ -1323,6 +1326,13 @@ void FAnimationBlueprintEditor::HandleObjectSelected(UObject* InObject)
 	SetDetailObject(InObject);
 }
 
+void FAnimationBlueprintEditor::HandleSelectionChanged(const TArrayView<TSharedPtr<ISkeletonTreeItem>>& InSelectedItems, ESelectInfo::Type InSelectInfo)
+{
+	TArray<UObject*> Objects;
+	Algo::TransformIf(InSelectedItems, Objects, [](const TSharedPtr<ISkeletonTreeItem>& InItem) { return InItem->GetObject() != nullptr; }, [](const TSharedPtr<ISkeletonTreeItem>& InItem) { return InItem->GetObject(); });
+	SetDetailObjects(Objects);
+}
+
 UObject* FAnimationBlueprintEditor::HandleGetObject()
 {
 	return GetEditingObject();
@@ -1457,5 +1467,14 @@ void FAnimationBlueprintEditor::HandleSetObjectBeingDebugged(UObject* InObject)
 	}
 }
 
+void FAnimationBlueprintEditor::HandlePreviewMeshChanged(USkeletalMesh* OldPreviewMesh, USkeletalMesh* NewPreviewMesh)
+{
+	UObject* Object = GetBlueprintObj()->GetObjectBeingDebugged();
+	if(Object)
+	{
+		HandleSetObjectBeingDebugged(Object);
+	}
+}
+
 #undef LOCTEXT_NAMESPACE
 
diff --git a/Engine/Source/Editor/AnimationBlueprintEditor/Private/AnimationBlueprintEditor.h b/Engine/Source/Editor/AnimationBlueprintEditor/Private/AnimationBlueprintEditor.h
--- a/Engine/Source/Editor/AnimationBlueprintEditor/Private/AnimationBlueprintEditor.h
+++ b/Engine/Source/Editor/AnimationBlueprintEditor/Private/AnimationBlueprintEditor.h
@@ -9,6 +9,7 @@
 #include "GraphEditor.h"
 #include "BlueprintEditor.h"
 #include "IAnimationBlueprintEditor.h"
+#include "ArrayView.h"
 
 class IPersonaToolkit;
 class IPersonaViewport;
@@ -17,6 +18,7 @@ class UAnimBlueprint;
 class UAnimGraphNode_Base;
 class UEdGraph;
 class USkeletalMesh;
+class ISkeletonTreeItem;
 
 struct FAnimationBlueprintEditorModes
 {
@@ -109,6 +111,7 @@ public:
 	/** Handle general object selection */
 	void HandleObjectsSelected(const TArray<UObject*>& InObjects);
 	void HandleObjectSelected(UObject* InObject);
+	void HandleSelectionChanged(const TArrayView<TSharedPtr<ISkeletonTreeItem>>& InSelectedItems, ESelectInfo::Type InSelectInfo);
 
 	/** Get the object to be displayed in the asset properties */
 	UObject* HandleGetObject();
@@ -251,6 +254,9 @@ private:
 	/** Handle a pin's default value changing be propagating it to the preview */
 	void HandlePinDefaultValueChanged(UEdGraphPin* InPinThatChanged);
 
+	/** Handle the preview mesh changing (so we can re-hook debug anim links etc.) */
+	void HandlePreviewMeshChanged(USkeletalMesh* OldPreviewMesh, USkeletalMesh* NewPreviewMesh);
+
 	/** The extender to pass to the level editor to extend it's window menu */
 	TSharedPtr<FExtender> MenuExtender;
 
diff --git a/Engine/Source/Editor/AnimationBlueprintEditor/Private/AnimationBlueprintEditorMode.cpp b/Engine/Source/Editor/AnimationBlueprintEditor/Private/AnimationBlueprintEditorMode.cpp
--- a/Engine/Source/Editor/AnimationBlueprintEditor/Private/AnimationBlueprintEditorMode.cpp
+++ b/Engine/Source/Editor/AnimationBlueprintEditor/Private/AnimationBlueprintEditorMode.cpp
@@ -121,7 +121,14 @@ FAnimationBlueprintEditorMode::FAnimationBlueprintEditorMode(const TSharedRef<FA
 	TabFactories.RegisterFactory(SkeletonEditorModule.CreateSkeletonTreeTabFactory(InAnimationBlueprintEditor, InAnimationBlueprintEditor->GetSkeletonTree()));
 
 	FPersonaModule& PersonaModule = FModuleManager::LoadModuleChecked<FPersonaModule>("Persona");
-	TabFactories.RegisterFactory(PersonaModule.CreatePersonaViewportTabFactory(InAnimationBlueprintEditor, InAnimationBlueprintEditor->GetSkeletonTree(), InAnimationBlueprintEditor->GetPersonaToolkit()->GetPreviewScene(), InAnimationBlueprintEditor->OnPostUndo, InAnimationBlueprintEditor, FOnViewportCreated(), true, false));
+
+	FPersonaViewportArgs ViewportArgs(InAnimationBlueprintEditor->GetSkeletonTree(), InAnimationBlueprintEditor->GetPersonaToolkit()->GetPreviewScene(), InAnimationBlueprintEditor->OnPostUndo);
+	ViewportArgs.BlueprintEditor = InAnimationBlueprintEditor;
+	ViewportArgs.bShowStats = false;
+
+	TabFactories.RegisterFactory(PersonaModule.CreatePersonaViewportTabFactory(InAnimationBlueprintEditor, ViewportArgs));
+
+
 	TabFactories.RegisterFactory(PersonaModule.CreateAdvancedPreviewSceneTabFactory(InAnimationBlueprintEditor, InAnimationBlueprintEditor->GetPersonaToolkit()->GetPreviewScene()));
 	TabFactories.RegisterFactory(PersonaModule.CreateAnimationAssetBrowserTabFactory(InAnimationBlueprintEditor, InAnimationBlueprintEditor->GetPersonaToolkit(), FOnOpenNewAsset::CreateSP(&InAnimationBlueprintEditor.Get(), &FAnimationBlueprintEditor::HandleOpenNewAsset), FOnAnimationSequenceBrowserCreated(), true));
 	TabFactories.RegisterFactory(PersonaModule.CreateAnimBlueprintPreviewTabFactory(InAnimationBlueprintEditor, InAnimationBlueprintEditor->GetPersonaToolkit()->GetPreviewScene()));
diff --git a/Engine/Source/Editor/AnimationEditor/Private/AnimationEditor.cpp b/Engine/Source/Editor/AnimationEditor/Private/AnimationEditor.cpp
--- a/Engine/Source/Editor/AnimationEditor/Private/AnimationEditor.cpp
+++ b/Engine/Source/Editor/AnimationEditor/Private/AnimationEditor.cpp
@@ -47,6 +47,8 @@
 #include "Developer/AssetTools/Public/IAssetTools.h"
 #include "Developer/AssetTools/Public/AssetToolsModule.h"
 #include "SequenceRecorderUtils.h"
+#include "ISkeletonTreeItem.h"
+#include "Algo/Transform.h"
 
 const FName AnimationEditorAppIdentifier = FName(TEXT("AnimationEditorApp"));
 
@@ -112,8 +114,8 @@ void FAnimationEditor::InitAnimationEditor(const EToolkitMode::Type Mode, const
 
 	PersonaToolkit->GetPreviewScene()->SetDefaultAnimationMode(EPreviewSceneDefaultAnimationMode::Animation);
 
-	FSkeletonTreeArgs SkeletonTreeArgs(OnPostUndo);
-	SkeletonTreeArgs.OnObjectSelected = FOnObjectSelected::CreateSP(this, &FAnimationEditor::HandleObjectSelected);
+	FSkeletonTreeArgs SkeletonTreeArgs;
+	SkeletonTreeArgs.OnSelectionChanged = FOnSkeletonTreeSelectionChanged::CreateSP(this, &FAnimationEditor::HandleSelectionChanged);
 	SkeletonTreeArgs.PreviewScene = PersonaToolkit->GetPreviewScene();
 
 	ISkeletonEditorModule& SkeletonEditorModule = FModuleManager::GetModuleChecked<ISkeletonEditorModule>("SkeletonEditor");
@@ -342,6 +344,16 @@ void FAnimationEditor::HandleObjectsSelected(const TArray<UObject*>& InObjects)
 	}
 }
 
+void FAnimationEditor::HandleSelectionChanged(const TArrayView<TSharedPtr<ISkeletonTreeItem>>& InSelectedItems, ESelectInfo::Type InSelectInfo)
+{
+	if (DetailsView.IsValid())
+	{
+		TArray<UObject*> Objects;
+		Algo::TransformIf(InSelectedItems, Objects, [](const TSharedPtr<ISkeletonTreeItem>& InItem) { return InItem->GetObject() != nullptr; }, [](const TSharedPtr<ISkeletonTreeItem>& InItem) { return InItem->GetObject(); });
+		DetailsView->SetObjects(Objects);
+	}
+}
+
 void FAnimationEditor::HandleObjectSelected(UObject* InObject)
 {
 	if (DetailsView.IsValid())
diff --git a/Engine/Source/Editor/AnimationEditor/Private/AnimationEditor.h b/Engine/Source/Editor/AnimationEditor/Private/AnimationEditor.h
--- a/Engine/Source/Editor/AnimationEditor/Private/AnimationEditor.h
+++ b/Engine/Source/Editor/AnimationEditor/Private/AnimationEditor.h
@@ -10,6 +10,7 @@
 #include "IAnimationEditor.h"
 #include "TickableEditorObject.h"
 #include "EditorUndoClient.h"
+#include "ArrayView.h"
 
 struct FAssetData;
 class FMenuBuilder;
@@ -21,6 +22,7 @@ class ISkeletonTree;
 class UAnimationAsset;
 class USkeletalMeshComponent;
 class UAnimSequence;
+class ISkeletonTreeItem;
 
 namespace AnimationEditorModes
 {
@@ -95,6 +97,8 @@ public:
 
 	void HandleAnimationSequenceBrowserCreated(const TSharedRef<class IAnimationSequenceBrowser>& InSequenceBrowser);
 
+	void HandleSelectionChanged(const TArrayView<TSharedPtr<ISkeletonTreeItem>>& InSelectedItems, ESelectInfo::Type InSelectInfo);
+
 	void HandleObjectSelected(UObject* InObject);
 
 	void HandleObjectsSelected(const TArray<UObject*>& InObjects);
diff --git a/Engine/Source/Editor/AnimationEditor/Private/AnimationEditorMode.cpp b/Engine/Source/Editor/AnimationEditor/Private/AnimationEditorMode.cpp
--- a/Engine/Source/Editor/AnimationEditor/Private/AnimationEditorMode.cpp
+++ b/Engine/Source/Editor/AnimationEditor/Private/AnimationEditorMode.cpp
@@ -22,7 +22,12 @@ FAnimationEditorMode::FAnimationEditorMode(TSharedRef<FWorkflowCentricApplicatio
 
 	FPersonaModule& PersonaModule = FModuleManager::LoadModuleChecked<FPersonaModule>("Persona");
 	TabFactories.RegisterFactory(PersonaModule.CreateDetailsTabFactory(InHostingApp, FOnDetailsCreated::CreateSP(&AnimationEditor.Get(), &FAnimationEditor::HandleDetailsCreated)));
-	TabFactories.RegisterFactory(PersonaModule.CreatePersonaViewportTabFactory(InHostingApp, InSkeletonTree, AnimationEditor->GetPersonaToolkit()->GetPreviewScene(), AnimationEditor->OnPostUndo, nullptr, FOnViewportCreated(), false, true));
+
+	FPersonaViewportArgs ViewportArgs(InSkeletonTree, AnimationEditor->GetPersonaToolkit()->GetPreviewScene(), AnimationEditor->OnPostUndo);
+	ViewportArgs.bShowTimeline = false;
+
+	TabFactories.RegisterFactory(PersonaModule.CreatePersonaViewportTabFactory(InHostingApp, ViewportArgs));
+
 	TabFactories.RegisterFactory(PersonaModule.CreateAdvancedPreviewSceneTabFactory(InHostingApp, AnimationEditor->GetPersonaToolkit()->GetPreviewScene()));
 	TabFactories.RegisterFactory(PersonaModule.CreateAnimationAssetBrowserTabFactory(InHostingApp, AnimationEditor->GetPersonaToolkit(), FOnOpenNewAsset::CreateSP(&AnimationEditor.Get(), &FAnimationEditor::HandleOpenNewAsset), FOnAnimationSequenceBrowserCreated::CreateSP(&AnimationEditor.Get(), &FAnimationEditor::HandleAnimationSequenceBrowserCreated), true));
 	TabFactories.RegisterFactory(PersonaModule.CreateAssetDetailsTabFactory(InHostingApp, FOnGetAsset::CreateSP(&AnimationEditor.Get(), &FAnimationEditor::HandleGetAsset), FOnDetailsCreated()));
diff --git a/Engine/Source/Editor/AnimationModifiers/Private/AnimationBlueprintLibrary.cpp b/Engine/Source/Editor/AnimationModifiers/Private/AnimationBlueprintLibrary.cpp
--- a/Engine/Source/Editor/AnimationModifiers/Private/AnimationBlueprintLibrary.cpp
+++ b/Engine/Source/Editor/AnimationModifiers/Private/AnimationBlueprintLibrary.cpp
@@ -924,7 +924,7 @@ void UAnimationBlueprintLibrary::RemoveCurve(UAnimSequence* AnimationSequence, F
 	}
 }
 
-void UAnimationBlueprintLibrary::RemoveAllCurveData(UAnimSequence* AnimationSequence, bool bRemoveNamesFromSkeleton /*= false*/)
+void UAnimationBlueprintLibrary::RemoveAllCurveData(UAnimSequence* AnimationSequence)
 {
 	if (AnimationSequence)
 	{
@@ -932,19 +932,6 @@ void UAnimationBlueprintLibrary::RemoveAllCurveData(UAnimSequence* AnimationSequ
 		AnimationSequence->RawCurveData.DeleteAllCurveData(ERawCurveTrackTypes::RCT_Vector);
 		AnimationSequence->RawCurveData.DeleteAllCurveData(ERawCurveTrackTypes::RCT_Transform);
 
-		if (bRemoveNamesFromSkeleton)
-		{
-			const FSmartNameMapping* CurveMapping = AnimationSequence->GetSkeleton()->GetSmartNameContainer(USkeleton::AnimCurveMappingName);
-			TArray<USkeleton::AnimCurveUID> CurveUids;
-			CurveMapping->FillUidArray(CurveUids);
-			AnimationSequence->GetSkeleton()->RemoveSmartnamesAndModify(USkeleton::AnimCurveMappingName, CurveUids);
-
-			const FSmartNameMapping* TrackCurveMapping = AnimationSequence->GetSkeleton()->GetSmartNameContainer(USkeleton::AnimTrackCurveMappingName);
-			CurveUids.Empty();
-			TrackCurveMapping->FillUidArray(CurveUids);
-			AnimationSequence->GetSkeleton()->RemoveSmartnamesAndModify(USkeleton::AnimTrackCurveMappingName, CurveUids);
-		}
-
 		AnimationSequence->bNeedsRebake = true;
 	}
 	else
@@ -1532,7 +1519,7 @@ void UAnimationBlueprintLibrary::GetBonePosesForTime(const UAnimSequence* Animat
 
 			if (RequiredBones.Num())
 			{
-				FBoneContainer BoneContainer(RequiredBones, false, *AnimationSequence->GetSkeleton());
+				FBoneContainer BoneContainer(RequiredBones, FCurveEvaluationOption(true), *AnimationSequence->GetSkeleton());
 				BoneContainer.SetUseSourceData(true);
 				BoneContainer.SetDisableRetargeting(true);
 				FCompactPose Pose;
diff --git a/Engine/Source/Editor/AnimationModifiers/Public/AnimationBlueprintLibrary.h b/Engine/Source/Editor/AnimationModifiers/Public/AnimationBlueprintLibrary.h
--- a/Engine/Source/Editor/AnimationModifiers/Public/AnimationBlueprintLibrary.h
+++ b/Engine/Source/Editor/AnimationModifiers/Public/AnimationBlueprintLibrary.h
@@ -220,7 +220,7 @@ public:
 
 	/** Removes all Animation Curve Data from the given Animation Sequence (Raw Animation Curves [Names] may not be removed from the Skeleton) */
 	UFUNCTION(BlueprintCallable, Category = "AnimationBlueprintLibrary|Curves")
-	static void RemoveAllCurveData(UAnimSequence* AnimationSequence, bool bRemoveNamesFromSkeleton = false);
+	static void RemoveAllCurveData(UAnimSequence* AnimationSequence);
 
 	/** Adds a Transformation Key to the specified Animation Curve inside of the given Animation Sequence */
 	UFUNCTION(BlueprintCallable, Category = "AnimationBlueprintLibrary|Curves")
diff --git a/Engine/Source/Editor/AudioEditor/Private/AssetTypeActions/AssetTypeActions_SoundSourceBus.cpp b/Engine/Source/Editor/AudioEditor/Private/AssetTypeActions/AssetTypeActions_SoundSourceBus.cpp
new file mode 100644
--- /dev/null
+++ b/Engine/Source/Editor/AudioEditor/Private/AssetTypeActions/AssetTypeActions_SoundSourceBus.cpp
@@ -0,0 +1,22 @@
+// Copyright 1998-2017 Epic Games, Inc. All Rights Reserved.
+
+#include "AssetTypeActions/AssetTypeActions_SoundSourceBus.h"
+#include "Misc/PackageName.h"
+#include "AssetData.h"
+#include "EditorStyleSet.h"
+#include "Factories/SoundSourceBusFactory.h"
+#include "IContentBrowserSingleton.h"
+#include "ContentBrowserModule.h"
+#include "PropertyCustomizationHelpers.h"
+#include "Widgets/Layout/SBox.h"
+#include "Widgets/Input/SButton.h"
+#include "Widgets/Images/SImage.h"
+#include "AudioEditorModule.h"
+
+
+UClass* FAssetTypeActions_SoundSourceBus::GetSupportedClass() const
+{
+	return USoundSourceBus::StaticClass();
+}
+
+
diff --git a/Engine/Source/Editor/AudioEditor/Private/AssetTypeActions/AssetTypeActions_SoundSourceBus.h b/Engine/Source/Editor/AudioEditor/Private/AssetTypeActions/AssetTypeActions_SoundSourceBus.h
new file mode 100644
--- /dev/null
+++ b/Engine/Source/Editor/AudioEditor/Private/AssetTypeActions/AssetTypeActions_SoundSourceBus.h
@@ -0,0 +1,23 @@
+// Copyright 1998-2017 Epic Games, Inc. All Rights Reserved.
+
+#pragma once
+
+#include "CoreMinimal.h"
+#include "Toolkits/IToolkitHost.h"
+#include "Sound/SoundSourceBus.h"
+#include "AssetTypeActions/AssetTypeActions_SoundBase.h"
+
+class FMenuBuilder;
+class USoundSourceBus;
+
+class FAssetTypeActions_SoundSourceBus : public FAssetTypeActions_SoundBase
+{
+public:
+	//~ Begin IAssetTypeActions Implementation
+	virtual FText GetName() const override { return NSLOCTEXT("AssetTypeActions", "AssetTypeActions_SoundSourceBus", "Sound Source Bus"); }
+	virtual FColor GetTypeColor() const override { return FColor(212, 97, 85); }
+	virtual UClass* GetSupportedClass() const override;
+	virtual bool CanFilter() override { return true; }
+	virtual bool IsImportedAsset() const override { return false; }
+	//~ End IAssetTypeActions Implementation
+};
diff --git a/Engine/Source/Editor/AudioEditor/Private/AudioEditorModule.cpp b/Engine/Source/Editor/AudioEditor/Private/AudioEditorModule.cpp
--- a/Engine/Source/Editor/AudioEditor/Private/AudioEditorModule.cpp
+++ b/Engine/Source/Editor/AudioEditor/Private/AudioEditorModule.cpp
@@ -31,6 +31,7 @@
 #include "AssetTypeActions/AssetTypeActions_ReverbEffect.h"
 #include "AssetTypeActions/AssetTypeActions_SoundSubmix.h"
 #include "AssetTypeActions/AssetTypeActions_SoundEffectPreset.h"
+#include "AssetTypeActions/AssetTypeActions_SoundSourceBus.h"
 #include "Utils.h"
 #include "UObject/UObjectIterator.h"
 
@@ -104,6 +105,7 @@ public:
 			AssetTools.RegisterAssetTypeActions(MakeShareable(new FAssetTypeActions_SoundEffectSubmixPreset));
 			AssetTools.RegisterAssetTypeActions(MakeShareable(new FAssetTypeActions_SoundEffectSourcePreset));
 			AssetTools.RegisterAssetTypeActions(MakeShareable(new FAssetTypeActions_SoundEffectSourcePresetChain));
+			AssetTools.RegisterAssetTypeActions(MakeShareable(new FAssetTypeActions_SoundSourceBus));
 		}
 	}
 
diff --git a/Engine/Source/Editor/AudioEditor/Private/Factories/SoundSourceBusFactory.cpp b/Engine/Source/Editor/AudioEditor/Private/Factories/SoundSourceBusFactory.cpp
new file mode 100644
--- /dev/null
+++ b/Engine/Source/Editor/AudioEditor/Private/Factories/SoundSourceBusFactory.cpp
@@ -0,0 +1,34 @@
+// Copyright 1998-2017 Epic Games, Inc. All Rights Reserved.
+
+#include "Factories/SoundSourceBusFactory.h"
+#include "Sound/SoundSourceBus.h"
+#include "AudioDeviceManager.h"
+#include "Engine/Engine.h"
+#include "Classes/Sound/AudioSettings.h"
+
+USoundSourceBusFactory::USoundSourceBusFactory(const FObjectInitializer& ObjectInitializer)
+	: Super(ObjectInitializer)
+{
+	SupportedClass = USoundSourceBus::StaticClass();
+	bCreateNew = true;
+	bEditorImport = false;
+	bEditAfterNew = true;
+}
+
+UObject* USoundSourceBusFactory::FactoryCreateNew(UClass* InClass, UObject* InParent, FName InName, EObjectFlags Flags, UObject* Context, FFeedbackContext* Warn)
+{
+	USoundSourceBus* SoundSourceBus = NewObject<USoundSourceBus>(InParent, InName, Flags);
+
+	FAudioDeviceManager* AudioDeviceManager = GEngine ? GEngine->GetAudioDeviceManager() : nullptr;
+	if (AudioDeviceManager)
+	{
+		AudioDeviceManager->InitSoundSubmixes();
+	}
+
+	return SoundSourceBus;
+}
+
+bool USoundSourceBusFactory::CanCreateNew() const
+{
+	return GetDefault<UAudioSettings>()->IsAudioMixerEnabled();
+}
diff --git a/Engine/Source/Editor/ClothPainter/Private/ClothPaintSettings.h b/Engine/Source/Editor/ClothPainter/Private/ClothPaintSettings.h
--- a/Engine/Source/Editor/ClothPainter/Private/ClothPaintSettings.h
+++ b/Engine/Source/Editor/ClothPainter/Private/ClothPaintSettings.h
@@ -56,11 +56,11 @@ public:
 
 protected:
 	/** When painting float/1D values, this is considered the zero or black point */
-	UPROPERTY(EditAnywhere, Category = View)
+	UPROPERTY(EditAnywhere, Category = View, meta = (UIMin = 0, UIMax = 100000, ClampMin = 0, ClampMax = 100000))
 	float ViewMin;
 
 	/** When painting float/1D values, this is considered the one or white point */
-	UPROPERTY(EditAnywhere, Category = View)
+	UPROPERTY(EditAnywhere, Category = View, meta = (UIMin = 0, UIMax = 100000, ClampMin = 0, ClampMax = 100000))
 	float ViewMax;
 
 public:
@@ -73,7 +73,7 @@ public:
 	UPROPERTY(EditAnywhere, AdvancedDisplay, Category = View)
 	bool bFlipNormal;
 
-	/** Whether to bull backfacing triangles when rendering the mesh preview */
+	/** Whether to cull backfacing triangles when rendering the mesh preview */
 	UPROPERTY(EditAnywhere, AdvancedDisplay, Category = View)
 	bool bCullBackface;
 
diff --git a/Engine/Source/Editor/ClothPainter/Private/ClothPaintTools.cpp b/Engine/Source/Editor/ClothPainter/Private/ClothPaintTools.cpp
--- a/Engine/Source/Editor/ClothPainter/Private/ClothPaintTools.cpp
+++ b/Engine/Source/Editor/ClothPainter/Private/ClothPaintTools.cpp
@@ -67,7 +67,7 @@ void FClothPaintTool_Brush::PaintAction(FPerVertexPaintActionArgs& InArgs, int32
 
 			float Value = SharedPainter->GetPropertyValue(VertexIndex);
 			const float BrushRadius = BrushSettings->GetBrushRadius();
-			MeshPaintHelpers::ApplyBrushToVertex(Position, InverseBrushMatrix, BrushRadius * BrushRadius, BrushSettings->BrushFalloffAmount, Settings->PaintValue, Value);
+			MeshPaintHelpers::ApplyBrushToVertex(Position, InverseBrushMatrix, BrushRadius, BrushSettings->BrushFalloffAmount, BrushSettings->BrushStrength, Settings->PaintValue, Value);
 			SharedPainter->SetPropertyValue(VertexIndex, Value);
 		}
 	}
@@ -135,27 +135,27 @@ void FClothPaintTool_Gradient::Render(USkeletalMeshComponent* InComponent, IMesh
 	
 	const FMatrix ComponentToWorldMatrix = InComponent->GetComponentTransform().ToMatrixWithScale();
 	
-	for (const MeshPaintHelpers::FPaintRay& PaintRay : PaintRays)
+	for (const int32& Index : GradientStartIndices)
 	{
-		for (const int32& Index : GradientStartIndices)
-		{
-			FVector Vertex;
-			InAdapter->GetVertexPosition(Index, Vertex);
+		FVector Vertex;
+		InAdapter->GetVertexPosition(Index, Vertex);
 
-			const FVector WorldPositionVertex = ComponentToWorldMatrix.TransformPosition(Vertex);
-			PDI->DrawPoint(WorldPositionVertex, FLinearColor::Green, VertexPointSize * 2.0f, SDPG_World);
-		}
+		const FVector WorldPositionVertex = ComponentToWorldMatrix.TransformPosition(Vertex);
+		PDI->DrawPoint(WorldPositionVertex, FLinearColor::Green, VertexPointSize * 2.0f, SDPG_World);
+	}
 	
-		for (const int32& Index : GradientEndIndices)
-		{
-			FVector Vertex;
-			InAdapter->GetVertexPosition(Index, Vertex);
+	for (const int32& Index : GradientEndIndices)
+	{
+		FVector Vertex;
+		InAdapter->GetVertexPosition(Index, Vertex);
 
-			const FVector WorldPositionVertex = ComponentToWorldMatrix.TransformPosition(Vertex);
-			PDI->DrawPoint(WorldPositionVertex, FLinearColor::Red, VertexPointSize * 2.0f, SDPG_World);
-		}
+		const FVector WorldPositionVertex = ComponentToWorldMatrix.TransformPosition(Vertex);
+		PDI->DrawPoint(WorldPositionVertex, FLinearColor::Red, VertexPointSize * 2.0f, SDPG_World);
+	}
 	
 	
+	for (const MeshPaintHelpers::FPaintRay& PaintRay : PaintRays)
+	{
 		const FHitResult& HitResult = SharedPainter->GetHitResult(PaintRay.RayStart, PaintRay.RayDirection);
 		if (HitResult.Component == InComponent)
 		{
diff --git a/Engine/Source/Editor/ClothPainter/Private/ClothPaintTools.h b/Engine/Source/Editor/ClothPainter/Private/ClothPaintTools.h
--- a/Engine/Source/Editor/ClothPainter/Private/ClothPaintTools.h
+++ b/Engine/Source/Editor/ClothPainter/Private/ClothPaintTools.h
@@ -73,11 +73,11 @@ public:
 	}
 
 	/** Value of the gradient at the start points */
-	UPROPERTY(EditAnywhere, Category = ToolSettings)
+	UPROPERTY(EditAnywhere, Category = ToolSettings, meta=(UIMin=0, UIMax=100, ClampMin=0, ClampMax=100000))
 	float GradientStartValue;
 
 	/** Value of the gradient at the end points */
-	UPROPERTY(EditAnywhere, Category = ToolSettings)
+	UPROPERTY(EditAnywhere, Category = ToolSettings, meta=(UIMin=0, UIMax=100, ClampMin=0, ClampMax=100000))
 	float GradientEndValue;
 
 	/** Enables the painting of selected points using a brush rather than just a point */
diff --git a/Engine/Source/Editor/ClothPainter/Private/ClothPainter.cpp b/Engine/Source/Editor/ClothPainter/Private/ClothPainter.cpp
--- a/Engine/Source/Editor/ClothPainter/Private/ClothPainter.cpp
+++ b/Engine/Source/Editor/ClothPainter/Private/ClothPainter.cpp
@@ -67,7 +67,7 @@ bool FClothPainter::PaintInternal(const FVector& InCameraOrigin, const FVector&
 {
 	bool bApplied = false;
 
-	if(SkeletalMeshComponent->SelectedClothingGuidForPainting.IsValid())
+	if(SkeletalMeshComponent->SelectedClothingGuidForPainting.IsValid() && !bShouldSimulate)
 	{
 		USkeletalMesh* SkelMesh = SkeletalMeshComponent->SkeletalMesh;
 
@@ -225,17 +225,17 @@ void FClothPainter::FinishPainting()
 	{		
 		EndTransaction();
 		Adapter->PostEdit();
-	}
-
-	if(SkeletalMeshComponent)
-	{
-		FComponentReregisterContext ReregisterContext(SkeletalMeshComponent);
 
-		if(USkeletalMesh* SkelMesh = SkeletalMeshComponent->SkeletalMesh)
+		if(SkeletalMeshComponent)
 		{
-			for(UClothingAssetBase* AssetBase : SkelMesh->MeshClothingAssets)
+			FComponentReregisterContext ReregisterContext(SkeletalMeshComponent);
+
+			if(USkeletalMesh* SkelMesh = SkeletalMeshComponent->SkeletalMesh)
 			{
-				AssetBase->InvalidateCachedData();
+				for(UClothingAssetBase* AssetBase : SkelMesh->MeshClothingAssets)
+				{
+					AssetBase->InvalidateCachedData();
+				}
 			}
 		}
 	}
@@ -312,7 +312,7 @@ void FClothPainter::Refresh()
 
 void FClothPainter::Render(const FSceneView* View, FViewport* Viewport, FPrimitiveDrawInterface* PDI)
 {
-	if(SelectedTool.IsValid() && SelectedTool->ShouldRenderInteractors())
+	if(SelectedTool.IsValid() && SelectedTool->ShouldRenderInteractors() && !bShouldSimulate)
 	{
 		RenderInteractors(View, Viewport, PDI, true, SDPG_Foreground);
 	}
diff --git a/Engine/Source/Editor/ClothPainter/Private/ClothPainterCommands.cpp b/Engine/Source/Editor/ClothPainter/Private/ClothPainterCommands.cpp
--- a/Engine/Source/Editor/ClothPainter/Private/ClothPainterCommands.cpp
+++ b/Engine/Source/Editor/ClothPainter/Private/ClothPainterCommands.cpp
@@ -2,9 +2,11 @@
 
 #include "ClothPainterCommands.h"
 
+#define LOCTEXT_NAMESPACE "ClothPaintCommands"
+
 void FClothPainterCommands::RegisterCommands()
 {
-
+	UI_COMMAND(TogglePaintMode, "Enable Cloth Paint", "Toggles between selection and clothing paint modes.", EUserInterfaceActionType::ToggleButton, FInputChord());
 }
 
 const FClothPainterCommands& FClothPainterCommands::Get()
@@ -12,3 +14,4 @@ const FClothPainterCommands& FClothPainterCommands::Get()
 	return TCommands<FClothPainterCommands>::Get();
 }
 
+#undef LOCTEXT_NAMESPACE
diff --git a/Engine/Source/Editor/ClothPainter/Private/ClothPainterCommands.h b/Engine/Source/Editor/ClothPainter/Private/ClothPainterCommands.h
--- a/Engine/Source/Editor/ClothPainter/Private/ClothPainterCommands.h
+++ b/Engine/Source/Editor/ClothPainter/Private/ClothPainterCommands.h
@@ -10,7 +10,7 @@ class FClothPainterCommands : public TCommands<FClothPainterCommands>
 {
 public:
 	FClothPainterCommands()
-		: TCommands<FClothPainterCommands>(TEXT("ClothPainter"), NSLOCTEXT("Contexts", "ClothPainter", "Cloth Painter"), NAME_None, FEditorStyle::GetStyleSetName())
+		: TCommands<FClothPainterCommands>(TEXT("ClothPainterTools"), NSLOCTEXT("Contexts", "ClothPainter", "Cloth Painter"), NAME_None, FEditorStyle::GetStyleSetName())
 	{
 
 	}
@@ -19,4 +19,5 @@ public:
 	static const FClothPainterCommands& Get();
 
 	/** Clothing commands */
+	TSharedPtr<FUICommandInfo> TogglePaintMode;
 };
\ No newline at end of file
diff --git a/Engine/Source/Editor/ClothPainter/Private/ClothPaintingModule.cpp b/Engine/Source/Editor/ClothPainter/Private/ClothPaintingModule.cpp
--- a/Engine/Source/Editor/ClothPainter/Private/ClothPaintingModule.cpp
+++ b/Engine/Source/Editor/ClothPainter/Private/ClothPaintingModule.cpp
@@ -16,11 +16,17 @@
 #include "ClothPaintSettingsCustomization.h"
 #include "Settings/EditorExperimentalSettings.h"
 #include "ClothPaintToolCommands.h"
+#include "ISkeletalMeshEditorModule.h"
+#include "MultiBoxBuilder.h"
+#include "ClothPainterCommands.h"
+#include "SDockTab.h"
 
 #define LOCTEXT_NAMESPACE "ClothPaintingModule"
 
 IMPLEMENT_MODULE(FClothPaintingModule, ClothPainter);
 
+DECLARE_DELEGATE_OneParam(FOnToggleClothPaintMode, bool);
+
 struct FClothPaintTabSummoner : public FWorkflowTabFactory
 {
 public:
@@ -30,7 +36,7 @@ public:
 	FClothPaintTabSummoner(TSharedPtr<class FAssetEditorToolkit> InHostingApp)
 		: FWorkflowTabFactory(TabName, InHostingApp)
 	{
-		TabLabel = LOCTEXT("ClothPaintTabLabel", "ClothPaint");
+		TabLabel = LOCTEXT("ClothPaintTabLabel", "Clothing");
 		TabIcon = FSlateIcon(FEditorStyle::GetStyleSetName(), "ClassIcon.SkeletalMesh");
 	}
 
@@ -48,21 +54,28 @@ public:
 	{
 		return MakeShareable(new FClothPaintTabSummoner(InAssetEditor));
 	}
+
+protected:
+
 };
 const FName FClothPaintTabSummoner::TabName = TEXT("ClothPainting");
 
 void FClothPaintingModule::StartupModule()
 {	
-	ExperimentalSettings = GetMutableDefault<UEditorExperimentalSettings>();
-	ExperimentalSettingsEventHandle = ExperimentalSettings->OnSettingChanged().AddRaw(this, &FClothPaintingModule::OnExperimentalSettingsChanged);
-
-	if(ExperimentalSettings->bClothingTools)
-	{
-		SetupMode();
-	}
+	SetupMode();
 
 	// Register any commands for the cloth painter
 	ClothPaintToolCommands::RegisterClothPaintToolCommands();
+	FClothPainterCommands::Register();
+
+	if(!SkelMeshEditorExtenderHandle.IsValid())
+	{
+		ISkeletalMeshEditorModule& SkelMeshEditorModule = FModuleManager::Get().LoadModuleChecked<ISkeletalMeshEditorModule>("SkeletalMeshEditor");
+		TArray<ISkeletalMeshEditorModule::FSkeletalMeshEditorToolbarExtender>& Extenders = SkelMeshEditorModule.GetAllSkeletalMeshEditorToolbarExtenders();
+		
+		Extenders.Add(ISkeletalMeshEditorModule::FSkeletalMeshEditorToolbarExtender::CreateRaw(this, &FClothPaintingModule::ExtendSkelMeshEditorToolbar));
+		SkelMeshEditorExtenderHandle = Extenders.Last().GetHandle();
+	}
 }
 
 void FClothPaintingModule::SetupMode()
@@ -86,39 +99,101 @@ TSharedRef<FApplicationMode> FClothPaintingModule::ExtendApplicationMode(const F
 	return InMode;
 }
 
-void FClothPaintingModule::OnExperimentalSettingsChanged(FName PropertyName)
+TSharedRef<FExtender> FClothPaintingModule::ExtendSkelMeshEditorToolbar(const TSharedRef<FUICommandList> InCommandList, TSharedRef<ISkeletalMeshEditor> InSkeletalMeshEditor)
+{
+	// Add toolbar extender
+	TSharedPtr<FExtender> ToolbarExtender = MakeShareable(new FExtender);
+
+	TWeakPtr<ISkeletalMeshEditor> Ptr(InSkeletalMeshEditor);
+
+	InCommandList->MapAction(FClothPainterCommands::Get().TogglePaintMode,
+		FExecuteAction::CreateRaw(this, &FClothPaintingModule::OnToggleMode, Ptr),
+		FCanExecuteAction(),
+		FIsActionChecked::CreateRaw(this, &FClothPaintingModule::GetIsPaintToolsButtonChecked, Ptr)
+	);
+
+	ToolbarExtender->AddToolBarExtension("Asset", EExtensionHook::After, InCommandList, FToolBarExtensionDelegate::CreateLambda(
+		[this, Ptr](FToolBarBuilder& Builder)
+	{
+		Builder.AddToolBarButton(
+			FClothPainterCommands::Get().TogglePaintMode, 
+			NAME_None, 
+			TAttribute<FText>::Create(TAttribute<FText>::FGetter::CreateRaw(this, &FClothPaintingModule::GetPaintToolsButtonText, Ptr)), 
+			TAttribute<FText>(), 
+			FSlateIcon(FEditorStyle::GetStyleSetName(), 
+			"LevelEditor.MeshPaintMode.TexturePaint"));
+	}));
+
+	return ToolbarExtender.ToSharedRef();
+}
+
+FText FClothPaintingModule::GetPaintToolsButtonText(TWeakPtr<ISkeletalMeshEditor> InSkeletalMeshEditor) const
 {
-	if(UEditorExperimentalSettings* ActualSettings = ExperimentalSettings.Get())
+	TSharedPtr<SClothPaintTab> ClothTab = GetActiveClothTab(InSkeletalMeshEditor, false);
+
+	if(ClothTab.IsValid())
 	{
-		if(PropertyName == GET_MEMBER_NAME_CHECKED(UEditorExperimentalSettings, bClothingTools))
+		if(ClothTab->IsPaintModeActive())
 		{
-			if(ActualSettings->bClothingTools)
-			{
-				SetupMode();
-			}
-			else
-			{
-				ShutdownMode();
-			}
+			return LOCTEXT("ToggleButton_Deactivate", "Deactivate Cloth Paint");
 		}
 	}
+
+	return LOCTEXT("ToggleButton_Activate", "Activate Cloth Paint");
 }
 
-void FClothPaintingModule::ShutdownModule()
+bool FClothPaintingModule::GetIsPaintToolsButtonChecked(TWeakPtr<ISkeletalMeshEditor> InSkeletalMeshEditor) const
 {
-	if(UEditorExperimentalSettings* ActualSettings = ExperimentalSettings.Get())
+	TSharedPtr<SClothPaintTab> ClothTab = GetActiveClothTab(InSkeletalMeshEditor, false);
+
+	if(ClothTab.IsValid())
 	{
-		if(ExperimentalSettings->bClothingTools)
-		{
-			ShutdownMode();
-		}
+		return ClothTab->IsPaintModeActive();
+	}
+
+	return false;
+}
 
-		ExperimentalSettings->OnSettingChanged().Remove(ExperimentalSettingsEventHandle);
+void FClothPaintingModule::OnToggleMode(TWeakPtr<ISkeletalMeshEditor> InSkeletalMeshEditor)
+{
+	TSharedPtr<SClothPaintTab> ClothTab = GetActiveClothTab(InSkeletalMeshEditor);
+
+	if(ClothTab.IsValid())
+	{
+		ClothTab->TogglePaintMode();
 	}
-	else
+}
+
+TSharedPtr<SClothPaintTab> FClothPaintingModule::GetActiveClothTab(TWeakPtr<ISkeletalMeshEditor> InSkeletalMeshEditor, bool bInvoke /*= true*/) const
+{
+	TSharedPtr<FTabManager> TabManager = InSkeletalMeshEditor.Pin()->GetTabManager();
+
+	if(bInvoke)
 	{
-		ShutdownMode();
+		TabManager->InvokeTab(FTabId(FClothPaintTabSummoner::TabName));
 	}
+
+	TSharedPtr<SDockTab> Tab = TabManager->FindExistingLiveTab(FTabId(FClothPaintTabSummoner::TabName));
+
+	if(Tab.IsValid())
+	{
+		TSharedPtr<SWidget> TabContent = Tab->GetContent();
+		TSharedPtr<SClothPaintTab> ClothingTab = StaticCastSharedPtr<SClothPaintTab>(TabContent);
+		return ClothingTab;
+	}
+
+	return nullptr;
+}
+
+void FClothPaintingModule::ShutdownModule()
+{
+	ShutdownMode();
+
+	// Remove skel mesh editor extenders
+	ISkeletalMeshEditorModule& SkelMeshEditorModule = FModuleManager::GetModuleChecked<ISkeletalMeshEditorModule>("SkeletalMeshEditor");
+	TArray<ISkeletalMeshEditorModule::FSkeletalMeshEditorToolbarExtender>& Extenders = SkelMeshEditorModule.GetAllSkeletalMeshEditorToolbarExtenders();
+	
+	Extenders.RemoveAll([=](const ISkeletalMeshEditorModule::FSkeletalMeshEditorToolbarExtender& InDelegate) {return InDelegate.GetHandle() == SkelMeshEditorExtenderHandle; });
 }
 
 void FClothPaintingModule::ShutdownMode()
diff --git a/Engine/Source/Editor/ClothPainter/Private/ClothingPaintEditMode.cpp b/Engine/Source/Editor/ClothPainter/Private/ClothingPaintEditMode.cpp
--- a/Engine/Source/Editor/ClothPainter/Private/ClothingPaintEditMode.cpp
+++ b/Engine/Source/Editor/ClothPainter/Private/ClothingPaintEditMode.cpp
@@ -108,7 +108,7 @@ void FClothingPaintEditMode::Exit()
 
 	if(PersonaToolkit.IsValid())
 	{
-		if(USkeletalMesh* SkelMesh = PersonaToolkit->GetPreviewMesh())
+		if(USkeletalMesh* SkelMesh = PersonaToolkit.Pin()->GetPreviewMesh())
 		{
 			for(TObjectIterator<USkeletalMeshComponent> It; It; ++It)
 			{
diff --git a/Engine/Source/Editor/ClothPainter/Private/ClothingPaintEditMode.h b/Engine/Source/Editor/ClothPainter/Private/ClothingPaintEditMode.h
--- a/Engine/Source/Editor/ClothPainter/Private/ClothingPaintEditMode.h
+++ b/Engine/Source/Editor/ClothPainter/Private/ClothingPaintEditMode.h
@@ -31,5 +31,5 @@ protected:
 	TSharedPtr<FClothPainter> ClothPainter;
 
 	class IPersonaPreviewScene* GetAnimPreviewScene() const;
-	class TSharedPtr<IPersonaToolkit> PersonaToolkit;
+	class TWeakPtr<IPersonaToolkit> PersonaToolkit;
 };
\ No newline at end of file
diff --git a/Engine/Source/Editor/ClothPainter/Private/SClothAssetSelector.cpp b/Engine/Source/Editor/ClothPainter/Private/SClothAssetSelector.cpp
--- a/Engine/Source/Editor/ClothPainter/Private/SClothAssetSelector.cpp
+++ b/Engine/Source/Editor/ClothPainter/Private/SClothAssetSelector.cpp
@@ -649,7 +649,7 @@ void SClothAssetSelector::Construct(const FArguments& InArgs, USkeletalMesh* InM
 				.VAlign(VAlign_Center)
 				[
 					SNew(STextBlock)
-					.Text(LOCTEXT("AssetExpander_Title", "Assets"))
+					.Text(LOCTEXT("AssetExpander_Title", "Clothing Data"))
 					.Font(FEditorStyle::GetFontStyle("DetailsView.CategoryFontStyle"))
 					.ShadowOffset(FVector2D(1.0f, 1.0f))
 				]
diff --git a/Engine/Source/Editor/ClothPainter/Private/SClothPaintTab.cpp b/Engine/Source/Editor/ClothPainter/Private/SClothPaintTab.cpp
--- a/Engine/Source/Editor/ClothPainter/Private/SClothPaintTab.cpp
+++ b/Engine/Source/Editor/ClothPainter/Private/SClothPaintTab.cpp
@@ -42,7 +42,10 @@ SClothPaintTab::~SClothPaintTab()
 {
 	if(ISkeletalMeshEditor* SkeletalMeshEditor = static_cast<ISkeletalMeshEditor*>(HostingApp.Pin().Get()))
 	{
-		SkeletalMeshEditor->GetAssetEditorModeManager()->ActivateDefaultMode();
+		if(FAssetEditorModeManager* ModeManager = SkeletalMeshEditor->GetAssetEditorModeManager())
+		{
+			ModeManager->ActivateDefaultMode();
+		}
 	}
 }
 
@@ -111,47 +114,6 @@ void SClothPaintTab::Construct(const FArguments& InArgs)
 		[
 			DetailsView->AsShared()
 		];
-
-		ContentBox->AddSlot()
-		.Padding(10.0f)
-		.AutoHeight()
-		[
-			SNew(SCheckBox)
-			.Type(ESlateCheckBoxType::ToggleButton)
-			.IsChecked_Lambda([=]() { return bPaintModeEnabled ? ECheckBoxState::Checked : ECheckBoxState::Unchecked;  })
-			.OnCheckStateChanged_Lambda([=](ECheckBoxState NewState) { bPaintModeEnabled = (NewState == ECheckBoxState::Checked); UpdatePaintTools(); })
-			.Style(&FEditorStyle::Get().GetWidgetStyle< FCheckBoxStyle >("ToggleButtonCheckbox"))
-			[
-				SNew(SBox)
-				.MinDesiredHeight(25.0f)
-				.MinDesiredWidth(100.0f)
-				.HAlign(HAlign_Center)
-				.VAlign(VAlign_Center)
-				.Padding(FMargin(4.0f, 4.0f, 4.0f, 4.0f))
-				[
-					SNew(SVerticalBox)							
-					+SVerticalBox::Slot()
-					.AutoHeight()
-					.Padding(0.0f, 4.0f)
-					[
-						SNew(SHorizontalBox)						
-						+SHorizontalBox::Slot()
-						.HAlign(HAlign_Center)
-						[
-							SNew(SImage)
-							.Image(TexturePaintIcon.GetIcon())
-						]
-						
-					]
-					+SVerticalBox::Slot()
-					.AutoHeight()
-					[
-						SNew(STextBlock)
-						.Text(FText::FromString("Enable Paint Tools"))
-					]
-				]
-			]
-		];
 	}
 }
 
@@ -160,6 +122,17 @@ void SClothPaintTab::Tick(const FGeometry& AllottedGeometry, const double InCurr
 	SCompoundWidget::Tick(AllottedGeometry, InCurrentTime, InDeltaTime);
 }
 
+void SClothPaintTab::TogglePaintMode()
+{
+	bPaintModeEnabled = !bPaintModeEnabled;
+	UpdatePaintTools();
+}
+
+bool SClothPaintTab::IsPaintModeActive() const
+{
+	return bPaintModeEnabled;
+}
+
 void SClothPaintTab::UpdatePaintTools()
 {
 	if (bPaintModeEnabled)
diff --git a/Engine/Source/Editor/ClothPainter/Private/SClothPaintTab.h b/Engine/Source/Editor/ClothPainter/Private/SClothPaintTab.h
--- a/Engine/Source/Editor/ClothPainter/Private/SClothPaintTab.h
+++ b/Engine/Source/Editor/ClothPainter/Private/SClothPaintTab.h
@@ -26,6 +26,12 @@ public:
 	void Construct(const FArguments& InArgs);
 	virtual void Tick(const FGeometry& AllottedGeometry, const double InCurrentTime, const float InDeltaTime) override;
 
+	/** Toggles cloth paint mode */
+	void TogglePaintMode();
+
+	/** Gets whether cloth paint mode is active */
+	bool IsPaintModeActive() const;
+
 protected:
 
 	/** Called as the tool selection changes to enable/disable painting */
diff --git a/Engine/Source/Editor/ClothPainter/Public/ClothPaintingModule.h b/Engine/Source/Editor/ClothPainter/Public/ClothPaintingModule.h
--- a/Engine/Source/Editor/ClothPainter/Public/ClothPaintingModule.h
+++ b/Engine/Source/Editor/ClothPainter/Public/ClothPaintingModule.h
@@ -8,6 +8,8 @@
 
 class FApplicationMode;
 class UEditorExperimentalSettings;
+class ISkeletalMeshEditor;
+class SClothPaintTab;
 
 const static FName PaintModeID = "ClothPaintMode";
 
@@ -24,19 +26,31 @@ public:
 	void ShutdownMode();
 
 protected:
+
+	// Extends the skeletal mesh editor mode
 	TSharedRef<FApplicationMode> ExtendApplicationMode(const FName ModeName, TSharedRef<FApplicationMode> InMode);
 protected:
+
 	TArray<TWeakPtr<FApplicationMode>> RegisteredApplicationModes;
 	FWorkflowApplicationModeExtender Extender;
 
 private:
 
-	// Weak ptr back to experimental settings as we'll need to unsubscribe later
-	TWeakObjectPtr<UEditorExperimentalSettings> ExperimentalSettings;
+	// Extends a skeletal mesh editor instance toolbar
+	TSharedRef<FExtender> ExtendSkelMeshEditorToolbar(const TSharedRef<FUICommandList> InCommandList, TSharedRef<ISkeletalMeshEditor> InSkeletalMeshEditor);
+
+	// Handle for the extender delegate we created
+	FDelegateHandle SkelMeshEditorExtenderHandle;
+
+	// Gets text for the enable paint tools button
+	FText GetPaintToolsButtonText(TWeakPtr<ISkeletalMeshEditor> InSkeletalMeshEditor) const;
+
+	// Whether paint mode is active
+	bool GetIsPaintToolsButtonChecked(TWeakPtr<ISkeletalMeshEditor> InSkeletalMeshEditor) const;
 
-	// Handler for the experimental settings changing so we can enable/disable the mode on the fly
-	void OnExperimentalSettingsChanged(FName PropertyName);
+	// Toggles paint mode on the clothing tab
+	void OnToggleMode(TWeakPtr<ISkeletalMeshEditor> InSkeletalMeshEditor);
 
-	// Handle for the above event so it can be unregistered when nexessary
-	FDelegateHandle ExperimentalSettingsEventHandle;
+	// Gets the current active clothing tab, will invoke (spawn or draw attention to) if bInvoke == true
+	TSharedPtr<SClothPaintTab> GetActiveClothTab(TWeakPtr<ISkeletalMeshEditor> InSkeletalMeshEditor, bool bInvoke = true) const;
 };
diff --git a/Engine/Source/Editor/DetailCustomizations/DetailCustomizations.Build.cs b/Engine/Source/Editor/DetailCustomizations/DetailCustomizations.Build.cs
--- a/Engine/Source/Editor/DetailCustomizations/DetailCustomizations.Build.cs
+++ b/Engine/Source/Editor/DetailCustomizations/DetailCustomizations.Build.cs
@@ -52,6 +52,7 @@ public class DetailCustomizations : ModuleRules
                 "SkeletonEditor",
                 "LevelSequence",
                 "AdvancedPreviewScene",
+                "AudioSettingsEditor",
             }
 		);
 
diff --git a/Engine/Source/Editor/DetailCustomizations/Private/AttenuationSettingsCustomizations.cpp b/Engine/Source/Editor/DetailCustomizations/Private/AttenuationSettingsCustomizations.cpp
--- a/Engine/Source/Editor/DetailCustomizations/Private/AttenuationSettingsCustomizations.cpp
+++ b/Engine/Source/Editor/DetailCustomizations/Private/AttenuationSettingsCustomizations.cpp
@@ -159,15 +159,17 @@ void FBaseAttenuationSettingsCustomization::CustomizeChildren(TSharedRef<IProper
 		.EditCondition(GetIsAttenuationEnabledAttribute(), nullptr);
 
 	// Get the attenuation category directly here otherwise our category is going to be incorrect for the following custom rows (e.g. "Vector" vs "Attenuation")
-	FText CategoryText = NSLOCTEXT("AttenuationSettings", "RadiusLabel", "AttenuationDistance");
+	FText CategoryText = NSLOCTEXT("AttenuationSettings", "AttenuationDistanceLabel", "AttenuationDistance");
 	FName AttenuationCategoryFName(*CategoryText.ToString());
 	IDetailCategoryBuilder& AttenuationCategory = LayoutBuilder.EditCategory(AttenuationCategoryFName);
 
-	AttenuationCategory.AddCustomRow(NSLOCTEXT("AttenuationSettings", "RadiusLabel", "Inner Radius"))
+	const FText RadiusLabel(NSLOCTEXT("AttenuationSettings", "RadiusLabel", "Inner Radius"));
+
+	AttenuationCategory.AddCustomRow(RadiusLabel)
 		.NameContent()
 		[
 			SNew(STextBlock)
-				.Text(NSLOCTEXT("AttenuationSettings", "RadiusLabel", "Inner Radius"))
+				.Text(RadiusLabel)
 				.ToolTipText(NSLOCTEXT("AttenuationSettings", "RadiusToolTip", "The radius that defines when sound attenuation begins (or when a custom attenuation curve begins). Sounds played at a distance less than this will not be attenuated."))
 				.Font(StructCustomizationUtils.GetRegularFont())
 				.IsEnabled(GetIsAttenuationEnabledAttribute())
diff --git a/Engine/Source/Editor/DetailCustomizations/Private/BodyInstanceCustomization.cpp b/Engine/Source/Editor/DetailCustomizations/Private/BodyInstanceCustomization.cpp
--- a/Engine/Source/Editor/DetailCustomizations/Private/BodyInstanceCustomization.cpp
+++ b/Engine/Source/Editor/DetailCustomizations/Private/BodyInstanceCustomization.cpp
@@ -1289,6 +1289,7 @@ void FBodyInstanceCustomizationHelper::UpdateFilters()
 
 			if(ObjectsCustomized[i]->IsA(USkeletalMeshComponent::StaticClass()))
 			{
+				bDisplayMass = false;
 				bDisplayAsyncScene = false;
 			}
 		}
@@ -1603,10 +1604,10 @@ EVisibility FBodyInstanceCustomizationHelper::IsDOFMode(EDOFMode::Type Mode) con
 
 void FBodyInstanceCustomizationHelper::AddMassInKg(IDetailCategoryBuilder& PhysicsCategory, TSharedRef<IPropertyHandle> BodyInstanceHandler)
 {
+	MassInKgOverrideHandle = BodyInstanceHandler->GetChildHandle(GET_MEMBER_NAME_CHECKED(FBodyInstance, MassInKgOverride)).ToSharedRef();
+
 	if (bDisplayMass)
 	{
-		MassInKgOverrideHandle = BodyInstanceHandler->GetChildHandle(GET_MEMBER_NAME_CHECKED(FBodyInstance, MassInKgOverride)).ToSharedRef();
-		
 		PhysicsCategory.AddProperty(MassInKgOverrideHandle).CustomWidget()
 		.NameContent()
 		[
@@ -1626,6 +1627,10 @@ void FBodyInstanceCustomizationHelper::AddMassInKg(IDetailCategoryBuilder& Physi
 			]
 		];
 	}
+	else
+	{
+		MassInKgOverrideHandle->MarkHiddenByCustomization();
+	}
 }
 
 void FBodyInstanceCustomizationHelper::AddMaxAngularVelocity(IDetailCategoryBuilder& PhysicsCategory, TSharedRef<IPropertyHandle> BodyInstanceHandler)
diff --git a/Engine/Source/Editor/DetailCustomizations/Private/BodySetupDetails.cpp b/Engine/Source/Editor/DetailCustomizations/Private/BodySetupDetails.cpp
--- a/Engine/Source/Editor/DetailCustomizations/Private/BodySetupDetails.cpp
+++ b/Engine/Source/Editor/DetailCustomizations/Private/BodySetupDetails.cpp
@@ -100,42 +100,6 @@ void FSkeletalBodySetupDetails::CustomizeDetails(IDetailLayoutBuilder& DetailBui
 		return bVisible;
 	};
 
-	auto AddProfileLambda = [ObjectsCustomizedLocal]() 
-	{
-		const FScopedTransaction Transaction(LOCTEXT("AddProfile", "Add Physical Animation Profile"));
-		bool bVisible = ObjectsCustomizedLocal.Num() > 0;
-		for (TWeakObjectPtr<UObject> WeakObj : ObjectsCustomizedLocal)
-		{
-			if (USkeletalBodySetup* BS = Cast<USkeletalBodySetup>(WeakObj.Get()))
-			{
-				FName ProfileName = BS->GetCurrentPhysicalAnimationProfileName();
-				if (!BS->FindPhysicalAnimationProfile(ProfileName))
-				{
-					BS->CurrentPhysicalAnimationProfile = FPhysicalAnimationProfile();
-					BS->AddPhysicalAnimationProfile(ProfileName);
-				}
-			}
-		}
-
-		return FReply::Handled();;
-	};
-
-	auto DeleteProfileLambda = [ObjectsCustomizedLocal]()
-	{
-		const FScopedTransaction Transaction(LOCTEXT("RemoveProfile", "Remove Physical Animation Profile"));
-		bool bVisible = ObjectsCustomizedLocal.Num() > 0;
-		for (TWeakObjectPtr<UObject> WeakObj : ObjectsCustomizedLocal)
-		{
-			if (USkeletalBodySetup* BS = Cast<USkeletalBodySetup>(WeakObj.Get()))
-			{
-				FName ProfileName = BS->GetCurrentPhysicalAnimationProfileName();
-				BS->RemovePhysicalAnimationProfile(ProfileName);
-			}
-		}
-
-		return FReply::Handled();;
-	};
-
 	TAttribute<EVisibility> PhysAnimVisible = TAttribute<EVisibility>::Create(TAttribute<EVisibility>::FGetter::CreateLambda([PhysAnimEditable]()
 	{
 		return PhysAnimEditable() == true ? EVisibility::Visible : EVisibility::Collapsed;
@@ -146,71 +110,26 @@ void FSkeletalBodySetupDetails::CustomizeDetails(IDetailLayoutBuilder& DetailBui
 		return PhysAnimEditable() == true ? EVisibility::Collapsed : EVisibility::Visible;
 	}));
 
-	TAttribute<bool> PhysAnimButtonEnabled = TAttribute<bool>::Create(TAttribute<bool>::FGetter::CreateLambda([ObjectsCustomizedLocal]()
-	{
-		for (TWeakObjectPtr<UObject> WeakObj : ObjectsCustomizedLocal)
-		{
-			if (USkeletalBodySetup* BS = Cast<USkeletalBodySetup>(WeakObj.Get()))
+	Cat.AddCustomRow(LOCTEXT("Profile", "Physical Animation Profile"))
+	[
+		SNew(SBorder)
+		.BorderImage(FEditorStyle::Get().GetBrush("ToolPanel.DarkGroupBorder"))
+		.Padding(4.0f)
+		[
+			SNew(SRichTextBlock)
+			.DecoratorStyleSet(&FEditorStyle::Get())
+			.Text_Lambda([ObjectsCustomizedLocal]()
 			{
-				if (BS->GetCurrentPhysicalAnimationProfileName() == NAME_None)
+				if (ObjectsCustomizedLocal.Num() > 0)
 				{
-					return false;
+					if (USkeletalBodySetup* BS = Cast<USkeletalBodySetup>(ObjectsCustomizedLocal[0].Get()))
+					{
+						return FText::Format(LOCTEXT("ProfileFormat", "Current Profile: <RichTextBlock.Bold>{0}</>"), FText::FromName(BS->GetCurrentPhysicalAnimationProfileName()));
+					}
 				}
-			}
-		}
-
-		return true;
-	}));
 
-	auto CurrentProfileTextLambda = [ObjectsCustomizedLocal]() -> FText
-	{
-		for (TWeakObjectPtr<UObject> WeakObj : ObjectsCustomizedLocal)
-		{
-			if (USkeletalBodySetup* BS = Cast<USkeletalBodySetup>(WeakObj.Get()))
-			{
-				return FText::FromName(BS->GetCurrentPhysicalAnimationProfileName());
-			}
-		}
-
-		return FText::FromName(NAME_None);
-	};
-
-	TAttribute<FText> CreateProfileLabel = TAttribute<FText>::Create(TAttribute<FText>::FGetter::CreateLambda([CurrentProfileTextLambda]()
-	{
-		return FText::Format(LOCTEXT("CreatePhysAnimLabel", "Add to profile: <RichTextBlock.Bold>{0}</>"), CurrentProfileTextLambda());
-	}));
-
-	Cat.AddCustomRow(LOCTEXT("NewPhysAnim", "NewPhysicalAnimationProfile"))
-	.Visibility(NewPhysAnimButtonVisible)
-	[
-		SNew(SVerticalBox)
-		+ SVerticalBox::Slot()
-		.HAlign(HAlign_Center)
-		.Padding(0, 10, 0, 10)
-		[
-			SNew(SBox)
-			.MinDesiredWidth(180)
-			.HeightOverride(32)
-			[
-				SNew(SButton)
-				.HAlign(HAlign_Center)
-				.IsEnabled(PhysAnimButtonEnabled)
-				.OnClicked_Lambda(AddProfileLambda)
-				[
-					SNew(SRichTextBlock)
-					.Text(CreateProfileLabel)
-					.DecoratorStyleSet(&FEditorStyle::Get())
-					.ToolTipText(LOCTEXT("NewPhysAnimButtonToolTip", "Add to current physical animation profile."))
-				]
-			]	
-		]
-		+ SVerticalBox::Slot()
-		.HAlign(HAlign_Center)
-		[
-			SNew(STextBlock)
-			.Text(LOCTEXT("NewPhysAnimLabel", "At least one body is not in the current physical animation profile."))
-			.Font(IDetailLayoutBuilder::GetDetailFont())
-			.ToolTipText(LOCTEXT("NewPhysAnimLabelTooltip", "Could not found in the current physical animation profile."))
+				return FText();
+			})
 		]
 	];
 
@@ -229,36 +148,6 @@ void FSkeletalBodySetupDetails::CustomizeDetails(IDetailLayoutBuilder& DetailBui
 			.Visibility(PhysAnimVisible);
 		}
 	}
-
-	TAttribute<FText> DeleteProfileLabel = TAttribute<FText>::Create(TAttribute<FText>::FGetter::CreateLambda([CurrentProfileTextLambda]()
-	{
-		return FText::Format(LOCTEXT("DeletePhysAnimLabel", "Remove from profile: <RichTextBlock.Bold>{0}</>"), CurrentProfileTextLambda());
-	}));
-
-	Cat.AddCustomRow(LOCTEXT("DeletePhysAnim", "DeletePhysicalAnimationProfile"))
-	.Visibility(PhysAnimVisible)
-	[
-		SNew(SVerticalBox)
-		+ SVerticalBox::Slot()
-		.Padding(0, 10, 0, 10)
-		.HAlign(HAlign_Center)
-		[
-			SNew(SBox)
-			.MinDesiredWidth(180)
-			[
-				SNew(SButton)
-				.HAlign(HAlign_Center)
-				.OnClicked_Lambda(DeleteProfileLambda)
-				[
-					SNew(SRichTextBlock)
-					.Text(DeleteProfileLabel)
-					.DecoratorStyleSet(&FEditorStyle::Get())
-					.ToolTipText(LOCTEXT("DeletePhysAnimButtonToolTip", "Removes the selected body from the current physical animation profile."))
-				]
-			]
-		]
-		
-	];
 }
 
 #undef LOCTEXT_NAMESPACE
diff --git a/Engine/Source/Editor/DetailCustomizations/Private/DetailCustomizations.cpp b/Engine/Source/Editor/DetailCustomizations/Private/DetailCustomizations.cpp
--- a/Engine/Source/Editor/DetailCustomizations/Private/DetailCustomizations.cpp
+++ b/Engine/Source/Editor/DetailCustomizations/Private/DetailCustomizations.cpp
@@ -83,6 +83,7 @@
 #include "CollisionProfileNameCustomization.h"
 #include "DocumentationActorDetails.h"
 #include "SoundBaseDetails.h"
+#include "SoundSourceBusDetails.h"
 #include "SoundWaveDetails.h"
 #include "AudioSettingsDetails.h"
 #include "DateTimeStructCustomization.h"
@@ -297,6 +298,7 @@ void FDetailCustomizationsModule::RegisterObjectCustomizations()
 	RegisterCustomClassLayout("AnimationAsset", FOnGetDetailCustomizationInstance::CreateStatic(&FAnimationAssetDetails::MakeInstance));
 
 	RegisterCustomClassLayout("SoundBase", FOnGetDetailCustomizationInstance::CreateStatic(&FSoundBaseDetails::MakeInstance));
+	RegisterCustomClassLayout("SoundSourceBus", FOnGetDetailCustomizationInstance::CreateStatic(&FSoundSourceBusDetails::MakeInstance));
 	RegisterCustomClassLayout("SoundWave", FOnGetDetailCustomizationInstance::CreateStatic(&FSoundWaveDetails::MakeInstance));
 	RegisterCustomClassLayout("DialogueWave", FOnGetDetailCustomizationInstance::CreateStatic(&FDialogueWaveDetails::MakeInstance));
 	RegisterCustomClassLayout("BodySetup", FOnGetDetailCustomizationInstance::CreateStatic(&FBodySetupDetails::MakeInstance));
diff --git a/Engine/Source/Editor/DetailCustomizations/Private/LinuxTargetSettingsDetails.cpp b/Engine/Source/Editor/DetailCustomizations/Private/LinuxTargetSettingsDetails.cpp
--- a/Engine/Source/Editor/DetailCustomizations/Private/LinuxTargetSettingsDetails.cpp
+++ b/Engine/Source/Editor/DetailCustomizations/Private/LinuxTargetSettingsDetails.cpp
@@ -269,6 +269,7 @@ void FLinuxTargetSettingsDetails::CustomizeDetails( IDetailLayoutBuilder& Detail
 		]
 	];
 
+	AudioPluginWidgetManager.BuildAudioCategory(DetailBuilder, EAudioPlatform::Linux);
 }
 
 
diff --git a/Engine/Source/Editor/DetailCustomizations/Private/LinuxTargetSettingsDetails.h b/Engine/Source/Editor/DetailCustomizations/Private/LinuxTargetSettingsDetails.h
--- a/Engine/Source/Editor/DetailCustomizations/Private/LinuxTargetSettingsDetails.h
+++ b/Engine/Source/Editor/DetailCustomizations/Private/LinuxTargetSettingsDetails.h
@@ -6,6 +6,7 @@
 #include "Styling/SlateColor.h"
 #include "Widgets/SWidget.h"
 #include "IDetailCustomization.h"
+#include "TargetPlatformAudioCustomization.h"
 
 class IDetailLayoutBuilder;
 class IPropertyHandle;
@@ -100,4 +101,7 @@ protected:
 private:
 	/** Reference to the target shader formats property view */
 	TSharedPtr<FLinuxTargetShaderFormatsPropertyDetails> TargetShaderFormatsDetails;
+
+	/** Manager for Audio Plugin widget. */
+	FAudioPluginWidgetManager AudioPluginWidgetManager;
 };
diff --git a/Engine/Source/Editor/DetailCustomizations/Private/MacTargetSettingsDetails.cpp b/Engine/Source/Editor/DetailCustomizations/Private/MacTargetSettingsDetails.cpp
--- a/Engine/Source/Editor/DetailCustomizations/Private/MacTargetSettingsDetails.cpp
+++ b/Engine/Source/Editor/DetailCustomizations/Private/MacTargetSettingsDetails.cpp
@@ -258,6 +258,8 @@ void FMacTargetSettingsDetails::CustomizeDetails( IDetailLayoutBuilder& DetailBu
 			.OnPostExternalImageCopy(FOnPostExternalImageCopy::CreateSP(this, &FMacTargetSettingsDetails::HandlePostExternalIconCopy))
 		]
 	];
+
+	AudioPluginWidgetManager.BuildAudioCategory(DetailBuilder, EAudioPlatform::Mac);
 }
 
 
diff --git a/Engine/Source/Editor/DetailCustomizations/Private/MacTargetSettingsDetails.h b/Engine/Source/Editor/DetailCustomizations/Private/MacTargetSettingsDetails.h
--- a/Engine/Source/Editor/DetailCustomizations/Private/MacTargetSettingsDetails.h
+++ b/Engine/Source/Editor/DetailCustomizations/Private/MacTargetSettingsDetails.h
@@ -5,6 +5,7 @@
 #include "CoreMinimal.h"
 #include "IDetailCustomization.h"
 #include "ShaderFormatsPropertyDetails.h"
+#include "TargetPlatformAudioCustomization.h"
 
 class IDetailLayoutBuilder;
 class IPropertyHandle;
@@ -54,4 +55,7 @@ private:
 	
 	/** Reference to the shader version property warning text box. */
 	TSharedPtr< SErrorText > ShaderVersionWarningTextBox;
+
+	/** Widget for platform specific audio plugins. */
+	FAudioPluginWidgetManager AudioPluginWidgetManager;
 };
diff --git a/Engine/Source/Editor/DetailCustomizations/Private/PhysicsConstraintComponentDetails.cpp b/Engine/Source/Editor/DetailCustomizations/Private/PhysicsConstraintComponentDetails.cpp
--- a/Engine/Source/Editor/DetailCustomizations/Private/PhysicsConstraintComponentDetails.cpp
+++ b/Engine/Source/Editor/DetailCustomizations/Private/PhysicsConstraintComponentDetails.cpp
@@ -343,6 +343,13 @@ void FPhysicsConstraintComponentDetails::AddAngularLimits(IDetailLayoutBuilder&
 	uint8 AngularLimitEnum[LCM_MAX] = { ACM_Free, LCM_Limited, LCM_Locked };
 	TSharedPtr<IPropertyHandle> AngularLimitProperties[] = { AngularSwing1MotionProperty, AngularSwing2MotionProperty, AngularTwistMotionProperty };
 
+	const FName AxisStyleNames[3] =
+	{
+		"PhysicsAssetEditor.RadioButtons.Red",
+		"PhysicsAssetEditor.RadioButtons.Red",
+		"PhysicsAssetEditor.RadioButtons.Green"
+	};
+
 	for (int32 PropertyIdx = 0; PropertyIdx < 3; ++PropertyIdx)
 	{
 		TSharedPtr<IPropertyHandle> CurProperty = AngularLimitProperties[PropertyIdx];
@@ -365,7 +372,7 @@ void FPhysicsConstraintComponentDetails::AddAngularLimits(IDetailLayoutBuilder&
 				.HAlign(HAlign_Left)
 				[
 					SNew(SCheckBox)
-					.Style(FEditorStyle::Get(), "RadioButton")
+					.Style(FEditorStyle::Get(), AxisStyleNames[PropertyIdx])
 					.IsChecked(this, &FPhysicsConstraintComponentDetails::IsLimitRadioChecked, CurProperty, AngularLimitEnum[0])
 					.OnCheckStateChanged(this, &FPhysicsConstraintComponentDetails::OnLimitRadioChanged, CurProperty, AngularLimitEnum[0])
 					.ToolTipText(AngularLimitOptionTooltips[0])
@@ -381,7 +388,7 @@ void FPhysicsConstraintComponentDetails::AddAngularLimits(IDetailLayoutBuilder&
 					.Padding(5, 0, 0, 0)
 					[
 						SNew(SCheckBox)
-						.Style(FEditorStyle::Get(), "RadioButton")
+						.Style(FEditorStyle::Get(), AxisStyleNames[PropertyIdx])
 						.IsChecked(this, &FPhysicsConstraintComponentDetails::IsLimitRadioChecked, CurProperty, AngularLimitEnum[1])
 						.OnCheckStateChanged(this, &FPhysicsConstraintComponentDetails::OnLimitRadioChanged, CurProperty, AngularLimitEnum[1])
 						.ToolTipText(AngularLimitOptionTooltips[1])
@@ -397,7 +404,7 @@ void FPhysicsConstraintComponentDetails::AddAngularLimits(IDetailLayoutBuilder&
 					.Padding(5, 0, 0, 0)
 					[
 						SNew(SCheckBox)
-						.Style(FEditorStyle::Get(), "RadioButton")
+						.Style(FEditorStyle::Get(), AxisStyleNames[PropertyIdx])
 						.IsChecked(this, &FPhysicsConstraintComponentDetails::IsLimitRadioChecked, CurProperty, AngularLimitEnum[2])
 						.OnCheckStateChanged(this, &FPhysicsConstraintComponentDetails::OnLimitRadioChanged, CurProperty, AngularLimitEnum[2])
 						.ToolTipText(AngularLimitOptionTooltips[2])
diff --git a/Engine/Source/Editor/DetailCustomizations/Private/SoundBaseDetails.cpp b/Engine/Source/Editor/DetailCustomizations/Private/SoundBaseDetails.cpp
--- a/Engine/Source/Editor/DetailCustomizations/Private/SoundBaseDetails.cpp
+++ b/Engine/Source/Editor/DetailCustomizations/Private/SoundBaseDetails.cpp
@@ -21,7 +21,10 @@ void FSoundBaseDetails::CustomizeDetails(IDetailLayoutBuilder& DetailBuilder)
 		Property = DetailBuilder.GetProperty("SourceEffectChain", USoundBase::StaticClass());
 		Property->MarkHiddenByCustomization();
 
-		Property = DetailBuilder.GetProperty("DefaultMasterReverbSendAmount", USoundBase::StaticClass());
+		Property = DetailBuilder.GetProperty("OutputToBusOnly", USoundBase::StaticClass());
+		Property->MarkHiddenByCustomization();
+
+		Property = DetailBuilder.GetProperty("BusSends", USoundBase::StaticClass());
 		Property->MarkHiddenByCustomization();
 	}
 }
diff --git a/Engine/Source/Editor/DetailCustomizations/Private/SoundSourceBusDetails.cpp b/Engine/Source/Editor/DetailCustomizations/Private/SoundSourceBusDetails.cpp
new file mode 100644
--- /dev/null
+++ b/Engine/Source/Editor/DetailCustomizations/Private/SoundSourceBusDetails.cpp
@@ -0,0 +1,29 @@
+// Copyright 1998-2017 Epic Games, Inc. All Rights Reserved.
+
+#include "SoundSourceBusDetails.h"
+#include "Settings/EditorExperimentalSettings.h"
+#include "Sound/SoundSourceBus.h"
+#include "DetailLayoutBuilder.h"
+#include "ScopedTransaction.h"
+#include "IDetailPropertyRow.h"
+#include "DetailCategoryBuilder.h"
+#include "Widgets/Input/SButton.h"
+#include "Engine/CurveTable.h"
+#include "DetailWidgetRow.h"
+#include "Widgets/Text/STextBlock.h"
+
+#define LOCTEXT_NAMESPACE "FSoundSourceBusDetails"
+
+TSharedRef<IDetailCustomization> FSoundSourceBusDetails::MakeInstance()
+{
+	return MakeShareable(new FSoundSourceBusDetails);
+}
+
+void FSoundSourceBusDetails::CustomizeDetails(IDetailLayoutBuilder& DetailBuilder)
+{
+	// Hide details of sound base that aren't relevant for buses
+
+
+}
+
+#undef LOCTEXT_NAMESPACE
diff --git a/Engine/Source/Editor/DetailCustomizations/Private/WindowsTargetSettingsDetails.cpp b/Engine/Source/Editor/DetailCustomizations/Private/WindowsTargetSettingsDetails.cpp
--- a/Engine/Source/Editor/DetailCustomizations/Private/WindowsTargetSettingsDetails.cpp
+++ b/Engine/Source/Editor/DetailCustomizations/Private/WindowsTargetSettingsDetails.cpp
@@ -330,16 +330,16 @@ void FWindowsTargetSettingsDetails::CustomizeDetails( IDetailLayoutBuilder& Deta
 					SNullWidget::NullWidget
 				]
 				.ContentPadding(FMargin(6.0f, 1.0f))
-					.MenuContent()
-					[
-						MakeAudioDeviceMenu(AudioDevicePropertyHandle)
-					]
+				.MenuContent()
+				[
+					MakeAudioDeviceMenu(AudioDevicePropertyHandle)
+				]
 				.ToolTipText(LOCTEXT("AudioDevicesButtonToolTip", "Pick from the list of available audio devices"))
 			]
 	];
+	AudioPluginWidgetManager.BuildAudioCategory(DetailBuilder, EAudioPlatform::Windows);
 }
 
-
 bool FWindowsTargetSettingsDetails::HandlePreExternalIconCopy(const FString& InChosenImage)
 {
 	return true;
diff --git a/Engine/Source/Editor/DetailCustomizations/Private/WindowsTargetSettingsDetails.h b/Engine/Source/Editor/DetailCustomizations/Private/WindowsTargetSettingsDetails.h
--- a/Engine/Source/Editor/DetailCustomizations/Private/WindowsTargetSettingsDetails.h
+++ b/Engine/Source/Editor/DetailCustomizations/Private/WindowsTargetSettingsDetails.h
@@ -6,6 +6,7 @@
 #include "Styling/SlateColor.h"
 #include "Widgets/SWidget.h"
 #include "IDetailCustomization.h"
+#include "TargetPlatformAudioCustomization.h"
 
 class IDetailLayoutBuilder;
 class IPropertyHandle;
@@ -100,4 +101,5 @@ protected:
 private:
 	/** Reference to the target shader formats property view */
 	TSharedPtr<FTargetShaderFormatsPropertyDetails> TargetShaderFormatsDetails;
+	FAudioPluginWidgetManager AudioPluginWidgetManager;
 };
