Fix Merge tool material issue

diff --git a/Engine/Source/Developer/MeshMergeUtilities/Private/MeshMergeUtilities.cpp b/Engine/Source/Developer/MeshMergeUtilities/Private/MeshMergeUtilities.cpp
--- a/Engine/Source/Developer/MeshMergeUtilities/Private/MeshMergeUtilities.cpp
+++ b/Engine/Source/Developer/MeshMergeUtilities/Private/MeshMergeUtilities.cpp
@@ -2476,6 +2476,7 @@ void FMeshMergeUtilities::MergeComponentsToStaticMesh(const TArray<UPrimitiveCom
 			UniqueMaterialIndices.AddUnique(MaterialIndex);
 		}
 
+		UniqueMaterialIndices.Sort();
 		for (int32 Index = 0; Index < UniqueMaterialIndices.Num(); ++Index)
 		{
 			const int32 SectionIndex = UniqueMaterialIndices[Index];
