fix from 4.22 cl:5317646
Fix MeshMergeUtilities when remapping polygon group

diff --git a/Engine/Source/Developer/MeshMergeUtilities/Private/MeshMergeUtilities.cpp b/Engine/Source/Developer/MeshMergeUtilities/Private/MeshMergeUtilities.cpp
--- a/Engine/Source/Developer/MeshMergeUtilities/Private/MeshMergeUtilities.cpp
+++ b/Engine/Source/Developer/MeshMergeUtilities/Private/MeshMergeUtilities.cpp
@@ -2662,7 +2662,8 @@ void FMeshMergeUtilities::CreateMergedRawMeshes(FMeshMergeDataTracker& InDataTra
 
 									// Note that at this point UniqueIndex is NOT a material index, but a unique section index!
 								}
-								else
+								
+								if(UniqueIndex == INDEX_NONE)
 								{
 									UniqueIndex = SourcePolygonGroupID.GetValue();
 								}
@@ -2765,7 +2766,8 @@ void FMeshMergeUtilities::CreateMergedRawMeshes(FMeshMergeDataTracker& InDataTra
 
 								// Note that at this point UniqueIndex is NOT a material index, but a unique section index!
 							}
-							else
+							
+							if(UniqueIndex == INDEX_NONE)
 							{
 								UniqueIndex = SourcePolygonGroupID.GetValue();
 							}
