Fix HLOD proxy build assert

diff --git a/Engine/Source/Developer/MeshMergeUtilities/Private/MeshMergeHelpers.cpp b/Engine/Source/Developer/MeshMergeUtilities/Private/MeshMergeHelpers.cpp
--- a/Engine/Source/Developer/MeshMergeUtilities/Private/MeshMergeHelpers.cpp
+++ b/Engine/Source/Developer/MeshMergeUtilities/Private/MeshMergeHelpers.cpp
@@ -243,6 +243,7 @@ void FMeshMergeHelpers::RetrieveMesh(const UStaticMeshComponent* StaticMeshCompo
 		// If removing degenerate triangles, ignore them when computing tangents.
 		TangentOptions |= FMeshDescriptionOperations::ETangentOptions::IgnoreDegenerateTriangles;
 	}
+	FMeshDescriptionOperations::CreatePolygonNTB(RawMesh, 0.0f);
 	FMeshDescriptionOperations::RecomputeNormalsAndTangentsIfNeeded(RawMesh, (FMeshDescriptionOperations::ETangentOptions)TangentOptions, BuildSettings.bUseMikkTSpace);
 }
 
@@ -401,6 +402,7 @@ void FMeshMergeHelpers::RetrieveMesh(const UStaticMesh* StaticMesh, int32 LODInd
 		// If removing degenerate triangles, ignore them when computing tangents.
 		TangentOptions |= FMeshDescriptionOperations::ETangentOptions::IgnoreDegenerateTriangles;
 	}
+	FMeshDescriptionOperations::CreatePolygonNTB(RawMesh, 0.0f);
 	FMeshDescriptionOperations::RecomputeNormalsAndTangentsIfNeeded(RawMesh, (FMeshDescriptionOperations::ETangentOptions)TangentOptions, BuildSettings.bUseMikkTSpace, (bImportedMesh && BuildSettings.bRecomputeNormals), (bImportedMesh && BuildSettings.bRecomputeTangents));
 }
 
diff --git a/Engine/Source/Developer/MeshMergeUtilities/Private/MeshMergeUtilities.cpp b/Engine/Source/Developer/MeshMergeUtilities/Private/MeshMergeUtilities.cpp
--- a/Engine/Source/Developer/MeshMergeUtilities/Private/MeshMergeUtilities.cpp
+++ b/Engine/Source/Developer/MeshMergeUtilities/Private/MeshMergeUtilities.cpp
@@ -1437,12 +1437,15 @@ void FMeshMergeUtilities::CreateProxyMesh(const TArray<UStaticMeshComponent*>& I
 			OutputMaterialsMap.MultiFind(MeshIndex, SectionAndOutputIndices);
 
 			TArray<int32> Remap;
+			TArray<int32> UniqueMaterialIndexes;
 			// Reorder loops 
 			for (const TPair<uint32, uint32>& IndexPair : SectionAndOutputIndices)
 			{
-				const int32 SectionIndex = IndexPair.Key;
+				//We are not using the IndexPair.Key since we want to keep the polygon group
+				//We instead find the section index by looking at unique IndexPair.Value
 				const int32 NewIndex = IndexPair.Value;
 
+				const int32 SectionIndex = UniqueMaterialIndexes.AddUnique(NewIndex);
 				if (Remap.Num() < (SectionIndex + 1))
 				{
 					Remap.SetNum(SectionIndex + 1);
@@ -1452,13 +1455,11 @@ void FMeshMergeUtilities::CreateProxyMesh(const TArray<UStaticMeshComponent*>& I
 			}
 			
 			TMap<FPolygonGroupID, FPolygonGroupID> RemapPolygonGroup;
-			int32 MaxRemapID = 0;
 			for (const FPolygonGroupID& PolygonGroupID : RawMesh.PolygonGroups().GetElementIDs())
 			{
 				checkf(Remap.IsValidIndex(PolygonGroupID.GetValue()), TEXT("Missing material bake output index entry for mesh(section)"));
 				int32 RemapID = Remap[PolygonGroupID.GetValue()];
 				RemapPolygonGroup.Add(PolygonGroupID, FPolygonGroupID(RemapID));
-				MaxRemapID = FMath::Max(MaxRemapID, RemapID);
 			}
 			FMeshDescriptionOperations::RemapPolygonGroups(RawMesh, RemapPolygonGroup);
 		}
