Fix the imposter mesh merge

diff --git a/Engine/Source/Developer/MeshMergeUtilities/Private/MeshMergeHelpers.cpp b/Engine/Source/Developer/MeshMergeUtilities/Private/MeshMergeHelpers.cpp
--- a/Engine/Source/Developer/MeshMergeUtilities/Private/MeshMergeHelpers.cpp
+++ b/Engine/Source/Developer/MeshMergeUtilities/Private/MeshMergeHelpers.cpp
@@ -1204,6 +1204,7 @@ void FMeshMergeHelpers::AppendRawMesh(FMeshDescription& InTarget, const FMeshDes
 	{
 		const FVertexID TargetVertexID = InTarget.CreateVertex();
 		SourceToTargetVertexID.Add(SourceVertexID, TargetVertexID);
+		TargetVertexPositions[TargetVertexID] = SourceVertexPositions[SourceVertexID];
 	}
 
 	//Append VertexInstances
@@ -1338,6 +1339,33 @@ void FMeshMergeHelpers::MergeImpostersToRawMesh(TArray<const UStaticMeshComponen
 				VertexInstanceUVs.Set(VertexInstanceID, UVTwoIndex, UVTwo);
 			}
 
+			TPolygonGroupAttributesRef<FName> SourcePolygonGroupImportedMaterialSlotNames = ImposterMesh.PolygonGroupAttributes().GetAttributesRef<FName>(MeshAttribute::PolygonGroup::ImportedMaterialSlotName);
+			TPolygonGroupAttributesRef<FName> TargetPolygonGroupImportedMaterialSlotNames = InRawMesh.PolygonGroupAttributes().GetAttributesRef<FName>(MeshAttribute::PolygonGroup::ImportedMaterialSlotName);
+
+			//Add the missing polygon group ID to the target(InRawMesh)
+			//Remap the source mesh(ImposterMesh) polygongroup to fit with the target polygon groups
+			TMap<FPolygonGroupID, FPolygonGroupID> RemapSourcePolygonGroup;
+			RemapSourcePolygonGroup.Reserve(ImposterMesh.PolygonGroups().Num());
+			for (const FPolygonGroupID& SourcePolygonGroupID : ImposterMesh.PolygonGroups().GetElementIDs())
+			{
+				FPolygonGroupID MatchTargetPolygonGroupID = FPolygonGroupID::Invalid;
+				for (const FPolygonGroupID& TargetPolygonGroupID : InRawMesh.PolygonGroups().GetElementIDs())
+				{
+					if (TargetPolygonGroupImportedMaterialSlotNames[TargetPolygonGroupID] == SourcePolygonGroupImportedMaterialSlotNames[SourcePolygonGroupID])
+					{
+						MatchTargetPolygonGroupID = TargetPolygonGroupID;
+						break;
+					}
+				}
+				if (MatchTargetPolygonGroupID == FPolygonGroupID::Invalid)
+				{
+					MatchTargetPolygonGroupID = InRawMesh.CreatePolygonGroup();
+					TargetPolygonGroupImportedMaterialSlotNames[MatchTargetPolygonGroupID] = SourcePolygonGroupImportedMaterialSlotNames[SourcePolygonGroupID];
+				}
+				RemapSourcePolygonGroup.Add(SourcePolygonGroupID, MatchTargetPolygonGroupID);
+			}
+			FMeshDescriptionOperations::RemapPolygonGroups(ImposterMesh, RemapSourcePolygonGroup);
+
 			FMeshMergeHelpers::AppendRawMesh(InRawMesh, ImposterMesh);
 		}
 	}
