Make sure the imposter materials are assign properly when there is more then one different imposter in the HLOD.

diff --git a/Engine/Source/Developer/MeshMergeUtilities/Private/MeshMergeHelpers.cpp b/Engine/Source/Developer/MeshMergeUtilities/Private/MeshMergeHelpers.cpp
--- a/Engine/Source/Developer/MeshMergeUtilities/Private/MeshMergeHelpers.cpp
+++ b/Engine/Source/Developer/MeshMergeUtilities/Private/MeshMergeHelpers.cpp
@@ -1299,6 +1299,7 @@ void FMeshMergeHelpers::MergeImpostersToRawMesh(TArray<const UStaticMeshComponen
 	// Ensure there are enough UV channels available to store the imposter data
 	//if (UVOneIndex != INDEX_NONE && UVOneIndex < (MAX_MESH_TEXTURE_COORDS - 2))
 	{
+		TMap<UMaterialInterface*, FPolygonGroupID> ImposterMaterialToPolygonGroupID;
 		for (const UStaticMeshComponent* Component : ImposterComponents)
 		{
 			// Retrieve imposter LOD mesh and material			
@@ -1313,9 +1314,10 @@ void FMeshMergeHelpers::MergeImpostersToRawMesh(TArray<const UStaticMeshComponen
 			TArray<FSectionInfo> Sections;
 			FMeshMergeHelpers::ExtractSections(Component, LODIndex, Sections);
 
+			TArray<int32> SectionImposterUniqueMaterialIndex;
 			for (FSectionInfo& Info : Sections)
 			{
-				OutImposterMaterials.AddUnique(Info.Material);
+				SectionImposterUniqueMaterialIndex.Add(OutImposterMaterials.AddUnique(Info.Material));
 			}
 
 			// Imposter magic, we're storing the actor world position and X scale spread across two UV channels
@@ -1346,21 +1348,17 @@ void FMeshMergeHelpers::MergeImpostersToRawMesh(TArray<const UStaticMeshComponen
 			//Remap the source mesh(ImposterMesh) polygongroup to fit with the target polygon groups
 			TMap<FPolygonGroupID, FPolygonGroupID> RemapSourcePolygonGroup;
 			RemapSourcePolygonGroup.Reserve(ImposterMesh.PolygonGroups().Num());
+			int32 SectionIndex = 0;
 			for (const FPolygonGroupID& SourcePolygonGroupID : ImposterMesh.PolygonGroups().GetElementIDs())
 			{
-				FPolygonGroupID MatchTargetPolygonGroupID = FPolygonGroupID::Invalid;
-				for (const FPolygonGroupID& TargetPolygonGroupID : InRawMesh.PolygonGroups().GetElementIDs())
-				{
-					if (TargetPolygonGroupImportedMaterialSlotNames[TargetPolygonGroupID] == SourcePolygonGroupImportedMaterialSlotNames[SourcePolygonGroupID])
-					{
-						MatchTargetPolygonGroupID = TargetPolygonGroupID;
-						break;
-					}
-				}
+				UMaterialInterface* MaterialUseBySection = OutImposterMaterials[SectionImposterUniqueMaterialIndex[SectionIndex++]];
+				FPolygonGroupID* ExistTargetPolygonGroupID = ImposterMaterialToPolygonGroupID.Find(MaterialUseBySection);
+				FPolygonGroupID MatchTargetPolygonGroupID = ExistTargetPolygonGroupID == nullptr ? FPolygonGroupID::Invalid : *ExistTargetPolygonGroupID;
 				if (MatchTargetPolygonGroupID == FPolygonGroupID::Invalid)
 				{
 					MatchTargetPolygonGroupID = InRawMesh.CreatePolygonGroup();
 					TargetPolygonGroupImportedMaterialSlotNames[MatchTargetPolygonGroupID] = SourcePolygonGroupImportedMaterialSlotNames[SourcePolygonGroupID];
+					ImposterMaterialToPolygonGroupID.Add(MaterialUseBySection, MatchTargetPolygonGroupID);
 				}
 				RemapSourcePolygonGroup.Add(SourcePolygonGroupID, MatchTargetPolygonGroupID);
 			}
