Fix material swap when merging imposter

diff --git a/Engine/Source/Developer/MeshMergeUtilities/Private/MeshMergeHelpers.cpp b/Engine/Source/Developer/MeshMergeUtilities/Private/MeshMergeHelpers.cpp
--- a/Engine/Source/Developer/MeshMergeUtilities/Private/MeshMergeHelpers.cpp
+++ b/Engine/Source/Developer/MeshMergeUtilities/Private/MeshMergeHelpers.cpp
@@ -1193,7 +1193,22 @@ void FMeshMergeHelpers::AppendRawMesh(FMeshDescription& InTarget, const FMeshDes
 		if (!InTarget.IsPolygonGroupValid(SourcePolygonGroupID))
 		{
 			InTarget.CreatePolygonGroupWithID(SourcePolygonGroupID);
-			TargetPolygonGroupImportedMaterialSlotNames[SourcePolygonGroupID] = SourcePolygonGroupImportedMaterialSlotNames[SourcePolygonGroupID];
+			const FName BaseName = SourcePolygonGroupImportedMaterialSlotNames[SourcePolygonGroupID];
+			FName CurrentTestName = BaseName;
+			int32 UniqueID = 1;
+			bool bUnique = true;
+			do 
+			{
+				for (const FPolygonGroupID PolygonGroupID : InTarget.PolygonGroups().GetElementIDs())
+				{
+					if (TargetPolygonGroupImportedMaterialSlotNames[PolygonGroupID] == CurrentTestName)
+					{
+						CurrentTestName = FName(*(BaseName.ToString() + FString::FromInt(UniqueID++)));
+						bUnique = false;
+					}
+				}
+			} while (!bUnique);
+			TargetPolygonGroupImportedMaterialSlotNames[SourcePolygonGroupID] = CurrentTestName;
 		}
 	}
 
@@ -1357,7 +1372,8 @@ void FMeshMergeHelpers::MergeImpostersToRawMesh(TArray<const UStaticMeshComponen
 				if (MatchTargetPolygonGroupID == FPolygonGroupID::Invalid)
 				{
 					MatchTargetPolygonGroupID = InRawMesh.CreatePolygonGroup();
-					TargetPolygonGroupImportedMaterialSlotNames[MatchTargetPolygonGroupID] = SourcePolygonGroupImportedMaterialSlotNames[SourcePolygonGroupID];
+					//use the material name to fill the imported material name. Material name will be unique
+					TargetPolygonGroupImportedMaterialSlotNames[MatchTargetPolygonGroupID] = MaterialUseBySection->GetFName();
 					ImposterMaterialToPolygonGroupID.Add(MaterialUseBySection, MatchTargetPolygonGroupID);
 				}
 				RemapSourcePolygonGroup.Add(SourcePolygonGroupID, MatchTargetPolygonGroupID);
