Need to Abs the X scale of the actor transform when generating imposter HLOD mesh.

diff --git a/Engine/Source/Developer/MeshMergeUtilities/Private/MeshMergeHelpers.cpp b/Engine/Source/Developer/MeshMergeUtilities/Private/MeshMergeHelpers.cpp
--- a/Engine/Source/Developer/MeshMergeUtilities/Private/MeshMergeHelpers.cpp
+++ b/Engine/Source/Developer/MeshMergeUtilities/Private/MeshMergeHelpers.cpp
@@ -1352,7 +1352,7 @@ void FMeshMergeHelpers::MergeImpostersToRawMesh(TArray<const UStaticMeshComponen
 				VertexInstanceUVs.Set(VertexInstanceID, UVOneIndex, UVOne);
 
 				UVTwo.X = ActorPosition.Z;
-				UVTwo.Y = ActorToWorld.GetScale3D().X;
+				UVTwo.Y = FMath::Abs(ActorToWorld.GetScale3D().X);
 				VertexInstanceUVs.Set(VertexInstanceID, UVTwoIndex, UVTwo);
 			}
 
