Fix a material ID issue with HLOD, by sorting the section mapping

diff --git a/Engine/Source/Developer/MeshMergeUtilities/Private/MeshMergeUtilities.cpp b/Engine/Source/Developer/MeshMergeUtilities/Private/MeshMergeUtilities.cpp
--- a/Engine/Source/Developer/MeshMergeUtilities/Private/MeshMergeUtilities.cpp
+++ b/Engine/Source/Developer/MeshMergeUtilities/Private/MeshMergeUtilities.cpp
@@ -1435,6 +1435,8 @@ void FMeshMergeUtilities::CreateProxyMesh(const TArray<UStaticMeshComponent*>& I
 
 			TArray<TPair<uint32, uint32>> SectionAndOutputIndices;
 			OutputMaterialsMap.MultiFind(MeshIndex, SectionAndOutputIndices);
+			//Make sure the section index are in the correct order
+			SectionAndOutputIndices.Sort([](const TPair<uint32, uint32>& A, const TPair<uint32, uint32>& B) { return (A.Key < B.Key); });
 
 			TArray<int32> Remap;
 			TArray<int32> UniqueMaterialIndexes;
