HLOD with imposter get the imposter section not assign to the correct material ID, and the proxy material is added for each imposter section.
The jira I put here must be re-tested since I change the fix for it.

diff --git a/Engine/Source/Developer/MeshMergeUtilities/Private/ProxyGenerationProcessor.cpp b/Engine/Source/Developer/MeshMergeUtilities/Private/ProxyGenerationProcessor.cpp
--- a/Engine/Source/Developer/MeshMergeUtilities/Private/ProxyGenerationProcessor.cpp
+++ b/Engine/Source/Developer/MeshMergeUtilities/Private/ProxyGenerationProcessor.cpp
@@ -208,7 +208,9 @@ void FProxyGenerationProcessor::ProcessJob(const FGuid& JobGuid, FProxyGeneratio
 	const bool bContainsImposters = Data->MergeData->ImposterComponents.Num() > 0;
 	FBox ImposterBounds(EForceInit::ForceInit);
 
-	auto RemoveVertexColorAndCommitMeshDescription = [&StaticMesh, &Data, &ProxyMaterial]()
+	TPolygonGroupAttributesConstRef<FName> PolygonGroupMaterialSlotName = Data->RawMesh.PolygonGroupAttributes().GetAttributesRef<FName>(MeshAttribute::PolygonGroup::ImportedMaterialSlotName);
+
+	auto RemoveVertexColorAndCommitMeshDescription = [&StaticMesh, &Data, &ProxyMaterial, &PolygonGroupMaterialSlotName]()
 	{
 		if (!Data->MergeData->InProxySettings.bAllowVertexColors)
 		{
@@ -228,14 +230,14 @@ void FProxyGenerationProcessor::ProcessJob(const FGuid& JobGuid, FProxyGeneratio
 		{
 			*MeshDescription = Data->RawMesh;
 
-			// material index / name mapping
-			TPolygonGroupAttributesConstRef<FName> PolygonGroupMaterialSlotName = MeshDescription->PolygonGroupAttributes().GetAttributesRef<FName>(MeshAttribute::PolygonGroup::ImportedMaterialSlotName);
-			for (const FPolygonGroupID PolygonGroupID : MeshDescription->PolygonGroups().GetElementIDs())
+			// Make sure the Proxy material have a valid ImportedMaterialSlotName
+			//The proxy material must be add only once and is always the first slot of the HLOD mesh
+			FStaticMaterial NewMaterial(ProxyMaterial);
+			if(MeshDescription->PolygonGroups().Num() > 0)
 			{
-				FStaticMaterial NewMaterial(ProxyMaterial);
-				NewMaterial.ImportedMaterialSlotName = PolygonGroupMaterialSlotName[PolygonGroupID];
-				StaticMesh->StaticMaterials.Add(NewMaterial);
+				NewMaterial.ImportedMaterialSlotName = PolygonGroupMaterialSlotName[MeshDescription->PolygonGroups().GetFirstValidID()];
 			}
+			StaticMesh->StaticMaterials.Add(NewMaterial);
 
 			StaticMesh->CommitMeshDescription(SourceModelIndex);
 		}
@@ -258,9 +260,16 @@ void FProxyGenerationProcessor::ProcessJob(const FGuid& JobGuid, FProxyGeneratio
 		}
 		RemoveVertexColorAndCommitMeshDescription();
 
+		
 		for (UMaterialInterface* Material : ImposterMaterials)
 		{
-			StaticMesh->StaticMaterials.Add(FStaticMaterial(Material));
+			//Set the ImportedMaterialSlotName in each imposter material
+			FStaticMaterial NewMaterial(Material);
+			if (Data->RawMesh.PolygonGroups().Num() > 0)
+			{
+				NewMaterial.ImportedMaterialSlotName = PolygonGroupMaterialSlotName[Data->RawMesh.PolygonGroups().GetFirstValidID()];
+			}
+			StaticMesh->StaticMaterials.Add(NewMaterial);
 		}
 	}
 	else
