HLOD Emissive bakes have pink artifacts or other objects emissive planar mapped
#fix Make sure we retain Magenta pixels for emissive baked textures, and not apply emissive scaling to them to make sure they are masked out correctly when merging the different textures into the atlas

diff --git a/Engine/Source/Developer/MaterialBaking/Private/MaterialBakingModule.cpp b/Engine/Source/Developer/MaterialBaking/Private/MaterialBakingModule.cpp
--- a/Engine/Source/Developer/MaterialBaking/Private/MaterialBakingModule.cpp
+++ b/Engine/Source/Developer/MaterialBaking/Private/MaterialBakingModule.cpp
@@ -407,13 +407,13 @@ void FMaterialBakingModule::ReadTextureOutput(FTextureRenderTargetResource* Rend
 
 					const FFloat16Color MagentaFloat16 = FFloat16Color(FLinearColor(1.0f, 0.0f, 1.0f));
 					FColor& Pixel8 = OutputColor[PixelX + YOffset];
-					/*if (Pixel16 == MagentaFloat16)
+					if (Pixel16 == MagentaFloat16)
 					{
 						Pixel8.R = 255;
 						Pixel8.G = 0;
 						Pixel8.B = 255;
 					}
-					else*/
+					else
 					{
 						Pixel8.R = (uint8)FMath::RoundToInt(Pixel16.R.GetFloat() * Scale);
 						Pixel8.G = (uint8)FMath::RoundToInt(Pixel16.G.GetFloat() * Scale);
diff --git a/Engine/Source/Developer/MeshMergeUtilities/Private/MeshMergeUtilities.cpp b/Engine/Source/Developer/MeshMergeUtilities/Private/MeshMergeUtilities.cpp
--- a/Engine/Source/Developer/MeshMergeUtilities/Private/MeshMergeUtilities.cpp
+++ b/Engine/Source/Developer/MeshMergeUtilities/Private/MeshMergeUtilities.cpp
@@ -1044,10 +1044,13 @@ float FMeshMergeUtilities::FlattenEmissivescale(TArray<struct FFlattenMaterial>&
 			{
 				for (FColor& Sample : Material.GetPropertySamples(EFlattenMaterialProperties::Emissive))
 				{
-					Sample.R = Sample.R * Multiplier;
-					Sample.G = Sample.G * Multiplier;
-					Sample.B = Sample.B * Multiplier;
-					Sample.A = Sample.A * Multiplier;
+					if (Sample != FColor::Magenta)
+					{
+						Sample.R = Sample.R * Multiplier;
+						Sample.G = Sample.G * Multiplier;
+						Sample.B = Sample.B * Multiplier;
+						Sample.A = Sample.A * Multiplier;
+					}
 				}
 			}
 		}
