Merging negative-scaled actors breaks materials
#fix No need to reverse the section indices when a static mesh has a mirrored transform as TransformRawMeshVertexData() is now calling ReverseAllPolygonFacing() in this specific case

diff --git a/Engine/Source/Developer/MeshMergeUtilities/Private/MeshMergeHelpers.cpp b/Engine/Source/Developer/MeshMergeUtilities/Private/MeshMergeHelpers.cpp
--- a/Engine/Source/Developer/MeshMergeUtilities/Private/MeshMergeHelpers.cpp
+++ b/Engine/Source/Developer/MeshMergeUtilities/Private/MeshMergeHelpers.cpp
@@ -64,7 +64,6 @@ void FMeshMergeHelpers::ExtractSections(const UStaticMeshComponent* Component, i
 #endif
 	}
 
-	const bool bMirrored = Component->GetComponentToWorld().GetDeterminant() < 0.f;
 	for (const FStaticMeshSection& MeshSection : StaticMesh->RenderData->LODResources[LODIndex].Sections)
 	{
 		// Retrieve material for this section
@@ -81,14 +80,6 @@ void FMeshMergeHelpers::ExtractSections(const UStaticMeshComponent* Component, i
 		SectionInfo.StartIndex = MeshSection.FirstIndex / 3;
 		SectionInfo.EndIndex = SectionInfo.StartIndex + MeshSection.NumTriangles;
 
-		// In case the object is mirrored the material indices/vertex data will be reversed in place, so we need to adjust the sections accordingly
-		if (bMirrored)
-		{
-			const uint32 NumTriangles = StaticMesh->RenderData->LODResources[LODIndex].GetNumTriangles();
-			SectionInfo.StartIndex = NumTriangles - SectionInfo.EndIndex;
-			SectionInfo.EndIndex = SectionInfo.StartIndex + MeshSection.NumTriangles;
-		}
-
 		if (MeshSection.bEnableCollision)
 		{
 			SectionInfo.EnabledProperties.Add(GET_MEMBER_NAME_CHECKED(FStaticMeshSection, bEnableCollision));
