Tentative, fix crash when generating HLOD with impostor

diff --git a/Engine/Source/Developer/MeshMergeUtilities/Private/MeshMergeUtilities.cpp b/Engine/Source/Developer/MeshMergeUtilities/Private/MeshMergeUtilities.cpp
--- a/Engine/Source/Developer/MeshMergeUtilities/Private/MeshMergeUtilities.cpp
+++ b/Engine/Source/Developer/MeshMergeUtilities/Private/MeshMergeUtilities.cpp
@@ -2773,10 +2773,12 @@ void FMeshMergeUtilities::CreateMergedRawMeshes(FMeshMergeDataTracker& InDataTra
 								// Note that at this point UniqueIndex is NOT a material index, but a unique section index!
 							}
 							
+							//Fallback
 							if(UniqueIndex == INDEX_NONE)
 							{
 								UniqueIndex = SourcePolygonGroupID.GetValue();
 							}
+
 							FPolygonGroupID TargetPolygonGroupID(UniqueIndex);
 							if (!TargetMesh.PolygonGroups().IsValid(TargetPolygonGroupID))
 							{
