Deprecate FRawMesh - Mesh merge utilities

diff --git a/Engine/Plugins/Experimental/ProxyLODPlugin/Source/ProxyLOD/Private/ProxyLODMeshTypes.cpp b/Engine/Plugins/Experimental/ProxyLODPlugin/Source/ProxyLOD/Private/ProxyLODMeshTypes.cpp
--- a/Engine/Plugins/Experimental/ProxyLODPlugin/Source/ProxyLOD/Private/ProxyLODMeshTypes.cpp
+++ b/Engine/Plugins/Experimental/ProxyLODPlugin/Source/ProxyLOD/Private/ProxyLODMeshTypes.cpp
@@ -48,22 +48,6 @@ void FMeshDescriptionAdapter::getIndexSpacePoint(size_t FaceNumber, size_t Corne
 
 
 // --- FMeshDescriptionArrayAdapter ----
-FMeshDescription* GetMeshDescriptionFromMeshMergeData(const FMeshMergeData* MergeData)
-{
-	FMeshDescription *MeshDescription = new FMeshDescription();
-	UStaticMesh::RegisterMeshAttributes(*MeshDescription);
-	TMap<int32, FName> MaterialMap;
-	if (MergeData->SourceStaticMesh)
-	{
-		for (int32 MaterialIndex = 0; MaterialIndex < MergeData->SourceStaticMesh->StaticMaterials.Num(); ++MaterialIndex)
-		{
-			MaterialMap.Add(MaterialIndex, MergeData->SourceStaticMesh->StaticMaterials[MaterialIndex].ImportedMaterialSlotName);
-		}
-	}
-	FMeshDescriptionOperations::ConvertFromRawMesh(*(MergeData->RawMesh), *MeshDescription, MaterialMap);
-	return MeshDescription;
-}
-
 FMeshDescriptionArrayAdapter::FMeshDescriptionArrayAdapter(const TArray<const FMeshMergeData*>& InMergeDataPtrArray)
 {
 	// Make a default transform.
@@ -76,7 +60,7 @@ FMeshDescriptionArrayAdapter::FMeshDescriptionArrayAdapter(const TArray<const FM
 	for (int32 MeshIdx = 0, MeshCount = InMergeDataPtrArray.Num(); MeshIdx < MeshCount; ++MeshIdx)
 	{
 		const FMeshMergeData* MergeData = InMergeDataPtrArray[MeshIdx];
-		FMeshDescription *RawMesh = GetMeshDescriptionFromMeshMergeData(MergeData);
+		FMeshDescription *RawMesh = MergeData->RawMesh;
 
 		PointCount += size_t(RawMesh->Vertices().Num());
 		for (const FPolygonID& PolygonID : RawMesh->Polygons().GetElementIDs())
@@ -107,7 +91,7 @@ FMeshDescriptionArrayAdapter::FMeshDescriptionArrayAdapter(const TArray<FMeshMer
 	{
 		const FMeshMergeData* MergeData = &InMergeDataArray[MeshIdx];
 		
-		FMeshDescription *RawMesh = GetMeshDescriptionFromMeshMergeData(MergeData);
+		FMeshDescription *RawMesh = MergeData->RawMesh;
 		PointCount += size_t(RawMesh->Vertices().Num());
 		for (const FPolygonID& PolygonID : RawMesh->Polygons().GetElementIDs())
 		{
@@ -135,7 +119,7 @@ FMeshDescriptionArrayAdapter::FMeshDescriptionArrayAdapter(const TArray<FMeshMer
 	{
 		const FMeshMergeData* MergeData = &InMergeDataArray[MeshIdx];
 		
-		FMeshDescription *RawMesh = GetMeshDescriptionFromMeshMergeData(MergeData);
+		FMeshDescription *RawMesh = MergeData->RawMesh;
 		PointCount += size_t(RawMesh->Vertices().Num());
 		for (const FPolygonID& PolygonID : RawMesh->Polygons().GetElementIDs())
 		{
@@ -213,23 +197,18 @@ void FMeshDescriptionArrayAdapter::UpdateMaterialsID()
 	{
 		FMeshDescription* MeshDescription = RawMeshArray[MeshIdx];
 
-		int32 FaceNum = MergeDataArray[MeshIdx]->RawMesh->FaceMaterialIndices.Num();
-		check(FaceNum == RawMeshArrayData[MeshIdx].TriangleCount);
+		check(MergeDataArray[MeshIdx]->RawMesh->Polygons().Num() == MeshDescription->Polygons().Num());
 		TMap<FPolygonGroupID, FPolygonGroupID> RemapGroup;
 		TArray<int32> UniqueMaterials;
-		int32 FaceIndex = 0;
 		for (const FPolygonID PolygonID : MeshDescription->Polygons().GetElementIDs())
 		{
-			int32 NewMaterialIndex = MergeDataArray[MeshIdx]->RawMesh->FaceMaterialIndices[FaceIndex];
-			if (!UniqueMaterials.Contains(NewMaterialIndex))
+			FPolygonGroupID NewPolygonGroupID = MergeDataArray[MeshIdx]->RawMesh->GetPolygonPolygonGroup(PolygonID);
+			if (!UniqueMaterials.Contains(NewPolygonGroupID.GetValue()))
 			{
-				UniqueMaterials.Add(NewMaterialIndex);
-				FPolygonGroupID NewPolygonGroupID = FPolygonGroupID(NewMaterialIndex);
+				UniqueMaterials.Add(NewPolygonGroupID.GetValue());
 				FPolygonGroupID OriginalPolygonGroupID = MeshDescription->GetPolygonPolygonGroup(PolygonID);
 				RemapGroup.Add(OriginalPolygonGroupID, NewPolygonGroupID);
 			}
-			//Update the face index
-			FaceIndex += MeshDescription->GetPolygonTriangles(PolygonID).Num();
 		}
 		//Remap the polygon group with the correct ID
 		FMeshDescriptionOperations::RemapPolygonGroups(*MeshDescription, RemapGroup);
diff --git a/Engine/Plugins/Experimental/ProxyLODPlugin/Source/ProxyLOD/Private/ProxyLODPlugin.cpp b/Engine/Plugins/Experimental/ProxyLODPlugin/Source/ProxyLOD/Private/ProxyLODPlugin.cpp
--- a/Engine/Plugins/Experimental/ProxyLODPlugin/Source/ProxyLOD/Private/ProxyLODPlugin.cpp
+++ b/Engine/Plugins/Experimental/ProxyLODPlugin/Source/ProxyLOD/Private/ProxyLODPlugin.cpp
@@ -1015,16 +1015,7 @@ void FVoxelizeMeshMerging::ProxyLOD(const FMeshMergeDataArray& InData, const FMe
 
 
 		// NB: FProxyGenerationProcessor::ProxyGenerationComplete
-		FRawMesh ConvertedRawMesh;
-		TMap<FName, int32> MaterialMap;
-		TPolygonGroupAttributesConstRef<FName> PolygonGroupNames = OutRawMesh.PolygonGroupAttributes().GetAttributesRef<FName>(MeshAttribute::PolygonGroup::ImportedMaterialSlotName);
-		int32 MaterialIndex = 0;
-		for (const FPolygonGroupID PolygonGroupID : OutRawMesh.PolygonGroups().GetElementIDs())
-		{
-			MaterialMap.Add(PolygonGroupNames[PolygonGroupID], MaterialIndex++);
-		}
-		FMeshDescriptionOperations::ConvertToRawMesh(OutRawMesh, ConvertedRawMesh, MaterialMap);
-		CompleteDelegate.ExecuteIfBound(ConvertedRawMesh, OutMaterial, InJobGUID);
+		CompleteDelegate.ExecuteIfBound(OutRawMesh, OutMaterial, InJobGUID);
 	}
 	else
 	{
diff --git a/Engine/Source/Developer/MaterialBaking/Public/MaterialBakingStructures.h b/Engine/Source/Developer/MaterialBaking/Public/MaterialBakingStructures.h
--- a/Engine/Source/Developer/MaterialBaking/Public/MaterialBakingStructures.h
+++ b/Engine/Source/Developer/MaterialBaking/Public/MaterialBakingStructures.h
@@ -28,12 +28,14 @@ struct FMaterialData
 struct FMeshData
 {
 	FMeshData()
-		: RawMesh(nullptr), Mesh(nullptr), bMirrored(false), VertexColorHash(0), TextureCoordinateIndex(0), LightMapIndex(0), LightMap(nullptr)
+		: RawMesh(nullptr), RawMeshDescription(nullptr), Mesh(nullptr), bMirrored(false), VertexColorHash(0), TextureCoordinateIndex(0), LightMapIndex(0), LightMap(nullptr)
 	{}
 
 	/** Ptr to raw mesh data to use for baking out the material data, if nullptr a standard quad is used */
 	FRawMesh* RawMesh;
 
+	FMeshDescription* RawMeshDescription;
+
 	/** Ptr to original static mesh this mesh data came from */
 	UStaticMesh* Mesh;
 
diff --git a/Engine/Source/Developer/MaterialUtilities/MaterialUtilities.Build.cs b/Engine/Source/Developer/MaterialUtilities/MaterialUtilities.Build.cs
--- a/Engine/Source/Developer/MaterialUtilities/MaterialUtilities.Build.cs
+++ b/Engine/Source/Developer/MaterialUtilities/MaterialUtilities.Build.cs
@@ -30,7 +30,8 @@ public class MaterialUtilities : ModuleRules
 
         PublicDependencyModuleNames.AddRange(
 			new string [] {
-				 "RawMesh",            
+				 "RawMesh",
+                 "MeshDescription",
 			}
 		);      
 
diff --git a/Engine/Source/Developer/MaterialUtilities/Private/MaterialUtilities.cpp b/Engine/Source/Developer/MaterialUtilities/Private/MaterialUtilities.cpp
--- a/Engine/Source/Developer/MaterialUtilities/Private/MaterialUtilities.cpp
+++ b/Engine/Source/Developer/MaterialUtilities/Private/MaterialUtilities.cpp
@@ -40,6 +40,11 @@
 #include "MaterialBakingStructures.h"
 #include "MaterialOptions.h"
 
+#include "MeshDescription.h"
+#include "MeshAttributes.h"
+#include "MeshAttributeArray.h"
+//#include "MeshDescriptionOperations.h"
+
 #if WITH_EDITOR
 #include "DeviceProfiles/DeviceProfile.h"
 #include "Tests/AutomationEditorCommon.h"
@@ -1714,9 +1719,11 @@ void FMaterialUtilities::RemapUniqueMaterialIndices(const TArray<FSectionInfo>&
 	{
 		for (int32 LODIndex = 0; LODIndex < MAX_STATIC_MESH_LODS; ++LODIndex)
 		{
+			TVertexAttributesRef<FVector> VertexPositions = InMeshData[MeshIndex].MeshLODData[LODIndex].RawMesh->VertexAttributes().GetAttributesRef<FVector>(MeshAttribute::Vertex::Position);
+
 			if (InMeshData[MeshIndex].bShouldExportLOD[LODIndex])
 			{
-				checkf(InMeshData[MeshIndex].MeshLODData[LODIndex].RawMesh->VertexPositions.Num(), TEXT("No vertex data found in mesh LOD"));
+				checkf(VertexPositions.GetNumElements(), TEXT("No vertex data found in mesh LOD"));
 				
 				const TArray<int32>& MeshMaterialMap = InMaterialMap[FMeshIdAndLOD(MeshIndex, LODIndex)];
 				int32 NumTexCoords = 0;
@@ -1762,9 +1769,11 @@ void FMaterialUtilities::RemapUniqueMaterialIndices(const TArray<FSectionInfo>&
 	{
 		for (int32 LODIndex = 0; LODIndex < MAX_STATIC_MESH_LODS; ++LODIndex)
 		{
+			TVertexAttributesRef<FVector> VertexPositions = InMeshData[MeshIndex].MeshLODData[LODIndex].RawMesh->VertexAttributes().GetAttributesRef<FVector>(MeshAttribute::Vertex::Position);
+
 			if (InMeshData[MeshIndex].bShouldExportLOD[LODIndex])
 			{
-				checkf(InMeshData[MeshIndex].MeshLODData[LODIndex].RawMesh->VertexPositions.Num(), TEXT("No vertex data found in mesh LOD"));
+				checkf(VertexPositions.GetNumElements(), TEXT("No vertex data found in mesh LOD"));
 		
 				const TArray<int32>& MeshMaterialMap = InMaterialMap[FMeshIdAndLOD(MeshIndex, LODIndex)];
 				TArray<int32>& NewMeshMaterialMap = OutMaterialMap.Add(FMeshIdAndLOD(MeshIndex, LODIndex));
diff --git a/Engine/Source/Developer/MeshDescriptionOperations/Private/MeshDescriptionOperations.cpp b/Engine/Source/Developer/MeshDescriptionOperations/Private/MeshDescriptionOperations.cpp
--- a/Engine/Source/Developer/MeshDescriptionOperations/Private/MeshDescriptionOperations.cpp
+++ b/Engine/Source/Developer/MeshDescriptionOperations/Private/MeshDescriptionOperations.cpp
@@ -305,7 +305,7 @@
 			}
 			else
 			{
-				DestinationRawMesh.FaceMaterialIndices[TriangleIndex] = 0;
+				DestinationRawMesh.FaceMaterialIndices[TriangleIndex] = PolygonGroupID.GetValue();
 			}
 			DestinationRawMesh.FaceSmoothingMasks[TriangleIndex] = 0; //Conversion of soft/hard to smooth mask is done after the geometry is converted
 			for (int32 Corner = 0; Corner < 3; ++Corner)
@@ -593,6 +593,167 @@
 	}
 }
 
+void FMeshDescriptionOperations::AppendMeshDescription(const FMeshDescription& SourceMesh, FMeshDescription& TargetMesh, FAppendSettings& AppendSettings)
+{
+	//Vertex Attributes
+	TVertexAttributesConstRef<FVector> SourceVertexPositions = SourceMesh.VertexAttributes().GetAttributesRef<FVector>(MeshAttribute::Vertex::Position);
+	TVertexAttributesConstRef<float> SourceVertexCornerSharpness = SourceMesh.VertexAttributes().GetAttributesRef<float>(MeshAttribute::Vertex::CornerSharpness);
+
+	TVertexAttributesRef<FVector> TargetVertexPositions = TargetMesh.VertexAttributes().GetAttributesRef<FVector>(MeshAttribute::Vertex::Position);
+	TVertexAttributesRef<float> TargetVertexCornerSharpness = TargetMesh.VertexAttributes().GetAttributesRef<float>(MeshAttribute::Vertex::CornerSharpness);
+
+	//Edge Attributes
+	TEdgeAttributesConstRef<bool> SourceEdgeHardnesses = SourceMesh.EdgeAttributes().GetAttributesRef<bool>(MeshAttribute::Edge::IsHard);
+	TEdgeAttributesConstRef<float> SourceEdgeCreaseSharpnesses = SourceMesh.EdgeAttributes().GetAttributesRef<float>(MeshAttribute::Edge::CreaseSharpness);
+
+	TEdgeAttributesRef<bool> TargetEdgeHardnesses = TargetMesh.EdgeAttributes().GetAttributesRef<bool>(MeshAttribute::Edge::IsHard);
+	TEdgeAttributesRef<float> TargetEdgeCreaseSharpnesses = TargetMesh.EdgeAttributes().GetAttributesRef<float>(MeshAttribute::Edge::CreaseSharpness);
+
+	//PolygonGroup Attributes
+	TPolygonGroupAttributesConstRef<FName> SourceImportedMaterialSlotNames = SourceMesh.PolygonGroupAttributes().GetAttributesRef<FName>(MeshAttribute::PolygonGroup::ImportedMaterialSlotName);
+	
+	TPolygonGroupAttributesRef<FName> TargetImportedMaterialSlotNames = TargetMesh.PolygonGroupAttributes().GetAttributesRef<FName>(MeshAttribute::PolygonGroup::ImportedMaterialSlotName);
+
+	//VertexInstance Attributes
+	TVertexInstanceAttributesConstRef<FVector> SourceVertexInstanceNormals = SourceMesh.VertexInstanceAttributes().GetAttributesRef<FVector>(MeshAttribute::VertexInstance::Normal);
+	TVertexInstanceAttributesConstRef<FVector> SourceVertexInstanceTangents = SourceMesh.VertexInstanceAttributes().GetAttributesRef<FVector>(MeshAttribute::VertexInstance::Tangent);
+	TVertexInstanceAttributesConstRef<float> SourceVertexInstanceBinormalSigns = SourceMesh.VertexInstanceAttributes().GetAttributesRef<float>(MeshAttribute::VertexInstance::BinormalSign);
+	TVertexInstanceAttributesConstRef<FVector4> SourceVertexInstanceColors = SourceMesh.VertexInstanceAttributes().GetAttributesRef<FVector4>(MeshAttribute::VertexInstance::Color);
+	TVertexInstanceAttributesConstRef<FVector2D> SourceVertexInstanceUVs = SourceMesh.VertexInstanceAttributes().GetAttributesRef<FVector2D>(MeshAttribute::VertexInstance::TextureCoordinate);
+
+	TVertexInstanceAttributesRef<FVector> TargetVertexInstanceNormals = TargetMesh.VertexInstanceAttributes().GetAttributesRef<FVector>(MeshAttribute::VertexInstance::Normal);
+	TVertexInstanceAttributesRef<FVector> TargetVertexInstanceTangents = TargetMesh.VertexInstanceAttributes().GetAttributesRef<FVector>(MeshAttribute::VertexInstance::Tangent);
+	TVertexInstanceAttributesRef<float> TargetVertexInstanceBinormalSigns = TargetMesh.VertexInstanceAttributes().GetAttributesRef<float>(MeshAttribute::VertexInstance::BinormalSign);
+	TVertexInstanceAttributesRef<FVector4> TargetVertexInstanceColors = TargetMesh.VertexInstanceAttributes().GetAttributesRef<FVector4>(MeshAttribute::VertexInstance::Color);
+	TVertexInstanceAttributesRef<FVector2D> TargetVertexInstanceUVs = TargetMesh.VertexInstanceAttributes().GetAttributesRef<FVector2D>(MeshAttribute::VertexInstance::TextureCoordinate);
+
+	//Copy into the target mesh
+	TargetMesh.ReserveNewVertices(SourceMesh.Vertices().Num());
+	TargetMesh.ReserveNewVertexInstances(SourceMesh.VertexInstances().Num());
+	TargetMesh.ReserveNewEdges(SourceMesh.Edges().Num());
+	TargetMesh.ReserveNewPolygons(SourceMesh.Polygons().Num());
+
+	if (SourceVertexInstanceUVs.GetNumIndices() > TargetVertexInstanceUVs.GetNumIndices())
+	{
+		TargetVertexInstanceUVs.SetNumIndices(SourceVertexInstanceUVs.GetNumIndices());
+	}
+
+	//PolygonGroups
+	PolygonGroupMap RemapPolygonGroup;
+	if (AppendSettings.PolygonGroupsDelegate.IsBound())
+	{
+		AppendSettings.PolygonGroupsDelegate.Execute(SourceMesh, TargetMesh, RemapPolygonGroup);
+	}
+	else
+	{
+		for (FPolygonGroupID SourcePolygonGroupID : SourceMesh.PolygonGroups().GetElementIDs())
+		{
+			FPolygonGroupID TargetMatchingID = FPolygonGroupID::Invalid;
+			for (FPolygonGroupID TargetPolygonGroupID : TargetMesh.PolygonGroups().GetElementIDs())
+			{
+				if (SourceImportedMaterialSlotNames[SourcePolygonGroupID] == TargetImportedMaterialSlotNames[TargetPolygonGroupID])
+				{
+					TargetMatchingID = TargetPolygonGroupID;
+					break;
+				}
+			}
+			if (TargetMatchingID == FPolygonGroupID::Invalid)
+			{
+				TargetMatchingID = TargetMesh.CreatePolygonGroup();
+				TargetImportedMaterialSlotNames[TargetMatchingID] = SourceImportedMaterialSlotNames[SourcePolygonGroupID];
+			}
+			RemapPolygonGroup.Add(SourcePolygonGroupID, TargetMatchingID);
+		}
+	}
+
+	//Vertices
+	TMap<FVertexID, FVertexID> SourceVertexIDRemap;
+	SourceVertexIDRemap.Reserve(SourceMesh.Vertices().Num());
+	for (FVertexID SourceVertexID : SourceMesh.Vertices().GetElementIDs())
+	{
+		FVertexID TargetVertexID = TargetMesh.CreateVertex();
+		TargetVertexPositions[TargetVertexID] = (SourceVertexPositions[SourceVertexID]-AppendSettings.MergedAssetPivot);
+		TargetVertexCornerSharpness[TargetVertexID] = SourceVertexCornerSharpness[SourceVertexID];
+		SourceVertexIDRemap.Add(SourceVertexID, TargetVertexID);
+	}
+
+	//Edges
+	TMap<FEdgeID, FEdgeID> SourceEdgeIDRemap;
+	SourceEdgeIDRemap.Reserve(SourceMesh.Edges().Num());
+	for (const FEdgeID& SourceEdgeID : SourceMesh.Edges().GetElementIDs())
+	{
+		const FMeshEdge& SourceEdge = SourceMesh.GetEdge(SourceEdgeID);
+		FEdgeID TargetEdgeID = TargetMesh.CreateEdge(SourceVertexIDRemap[SourceEdge.VertexIDs[0]], SourceVertexIDRemap[SourceEdge.VertexIDs[1]]);
+		TargetEdgeHardnesses[TargetEdgeID] = SourceEdgeHardnesses[SourceEdgeID];
+		TargetEdgeCreaseSharpnesses[TargetEdgeID] = SourceEdgeCreaseSharpnesses[SourceEdgeID];
+		SourceEdgeIDRemap.Add(SourceEdgeID, TargetEdgeID);
+	}
+
+	//VertexInstances
+	TMap<FVertexInstanceID, FVertexInstanceID> SourceVertexInstanceIDRemap;
+	SourceVertexInstanceIDRemap.Reserve(SourceMesh.VertexInstances().Num());
+	for (const FVertexInstanceID& SourceVertexInstanceID : SourceMesh.VertexInstances().GetElementIDs())
+	{
+		FVertexInstanceID TargetVertexInstanceID = TargetMesh.CreateVertexInstance(SourceVertexIDRemap[SourceMesh.GetVertexInstanceVertex(SourceVertexInstanceID)]);
+		SourceVertexInstanceIDRemap.Add(SourceVertexInstanceID, TargetVertexInstanceID);
+		
+		TargetVertexInstanceNormals[TargetVertexInstanceID] = SourceVertexInstanceNormals[SourceVertexInstanceID];
+		TargetVertexInstanceTangents[TargetVertexInstanceID] = SourceVertexInstanceTangents[SourceVertexInstanceID];
+		TargetVertexInstanceBinormalSigns[TargetVertexInstanceID] = SourceVertexInstanceBinormalSigns[SourceVertexInstanceID];
+
+		if (AppendSettings.bMergeVertexColor)
+		{
+			TargetVertexInstanceColors[TargetVertexInstanceID] = SourceVertexInstanceColors[SourceVertexInstanceID];
+		}
+
+		for (int32 UVChannelIndex = 0; UVChannelIndex < SourceVertexInstanceUVs.GetNumIndices(); ++UVChannelIndex)
+		{
+			TargetVertexInstanceUVs.Set(TargetVertexInstanceID, UVChannelIndex, SourceVertexInstanceUVs.Get(SourceVertexInstanceID, UVChannelIndex));
+		}
+	}
+
+	//Polygons
+	for (const FPolygonID& SourcePolygonID : SourceMesh.Polygons().GetElementIDs())
+	{
+		const FMeshPolygon& SourcePolygon = SourceMesh.GetPolygon(SourcePolygonID);
+		//Find the polygonGroupID
+		FPolygonGroupID TargetPolygonGroupID = RemapPolygonGroup[SourcePolygon.PolygonGroupID];
+
+		int32 PolygonVertexCount = SourcePolygon.PerimeterContour.VertexInstanceIDs.Num();
+		//Create a contour
+		TArray<FMeshDescription::FContourPoint> Contours;
+		// Add the edges of this polygon
+		for (uint32 PolygonEdgeNumber = 0; PolygonEdgeNumber < (uint32)PolygonVertexCount; ++PolygonEdgeNumber)
+		{
+			int32 ContourPointIndex = Contours.AddDefaulted();
+			FMeshDescription::FContourPoint& ContourPoint = Contours[ContourPointIndex];
+			//Find the matching edge ID
+			uint32 CornerIndices[2];
+			CornerIndices[0] = (PolygonEdgeNumber + 0) % PolygonVertexCount;
+			CornerIndices[1] = (PolygonEdgeNumber + 1) % PolygonVertexCount;
+
+			FVertexID SourceVertexIDs[2];
+			SourceVertexIDs[0] = SourceMesh.GetVertexInstanceVertex(SourcePolygon.PerimeterContour.VertexInstanceIDs[CornerIndices[0]]);
+			SourceVertexIDs[1] = SourceMesh.GetVertexInstanceVertex(SourcePolygon.PerimeterContour.VertexInstanceIDs[CornerIndices[1]]);
+
+			FEdgeID SourceEdgeID = SourceMesh.GetVertexPairEdge(SourceVertexIDs[0], SourceVertexIDs[1]);
+			check(SourceEdgeID != FEdgeID::Invalid)
+			FEdgeID TargetEdgeID = SourceEdgeIDRemap[SourceEdgeID];
+			ContourPoint.EdgeID = TargetEdgeID;
+			ContourPoint.VertexInstanceID = SourceVertexInstanceIDRemap[SourcePolygon.PerimeterContour.VertexInstanceIDs[CornerIndices[0]]];
+
+			TargetEdgeCreaseSharpnesses[TargetEdgeID] = SourceEdgeCreaseSharpnesses[SourceEdgeID];
+			TargetEdgeHardnesses[TargetEdgeID] = SourceEdgeHardnesses[SourceEdgeID];
+		}
+
+		// Insert a polygon into the mesh
+		const FPolygonID TargetPolygonID = TargetMesh.CreatePolygon(TargetPolygonGroupID, Contours);
+		//Triangulate the polygon
+		FMeshPolygon& Polygon = TargetMesh.GetPolygon(TargetPolygonID);
+		TargetMesh.ComputePolygonTriangulation(TargetPolygonID, Polygon.Triangles);
+	}
+}
+
 //////////////////////////////////////////////////////////////////////////
 // Normals tangents and Bi-normals
 void FMeshDescriptionOperations::RecomputeNormalsAndTangentsIfNeeded(FMeshDescription& MeshDescription, ETangentOptions TangentOptions, bool bUseMikkTSpace, bool bForceRecomputeNormals, bool bForceRecomputeTangents)
@@ -1227,10 +1388,112 @@
 	}
 }
 
-bool FMeshDescriptionOperations::GenerateUniqueUVsForStaticMesh(const FMeshDescription& MeshDescription, int32 TextureResolution, TArray<FVector2D>& OutTexCoords)
+bool FMeshDescriptionOperations::GenerateUniqueUVsForStaticMesh(const FMeshDescription& MeshDescription, int32 TextureResolution, bool bMergeIdenticalMaterials, TArray<FVector2D>& OutTexCoords)
 {
 	// Create a copy of original mesh (only copy necessary data)
 	FMeshDescription DuplicateMeshDescription(MeshDescription);
+
+	TMap<FVertexInstanceID, FVertexInstanceID> RemapVertexInstance;
+	//Remove the identical material
+	if (bMergeIdenticalMaterials)
+	{
+		TVertexInstanceAttributesConstRef<FVector2D> VertexInstanceUVs = DuplicateMeshDescription.VertexInstanceAttributes().GetAttributesRef<FVector2D>(MeshAttribute::VertexInstance::TextureCoordinate);
+		TArray<FPolygonID> ToDeletePolygons;
+		RemapVertexInstance.Reserve(DuplicateMeshDescription.VertexInstances().Num());
+		TArray<FPolygonID> UniquePolygons;
+		for (FPolygonID RefPolygonID : DuplicateMeshDescription.Polygons().GetElementIDs())
+		{
+			FPolygonGroupID RefPolygonGroupID = DuplicateMeshDescription.GetPolygonPolygonGroup(RefPolygonID);
+			const TArray<FVertexInstanceID>& RefVertexInstances = DuplicateMeshDescription.GetPolygonPerimeterVertexInstances(RefPolygonID);
+			TArray<FVector2D> RefUVs;
+			for (FVertexInstanceID RefVertexInstanceID : RefVertexInstances)
+			{
+				RefUVs.Add(VertexInstanceUVs[RefVertexInstanceID]);
+			}
+			FPolygonID MatchPolygonID = FPolygonID::Invalid;
+			for (FPolygonID TestPolygonID : UniquePolygons)
+			{
+				FPolygonGroupID TestPolygonGroupID = DuplicateMeshDescription.GetPolygonPolygonGroup(TestPolygonID);
+				if (TestPolygonGroupID != RefPolygonGroupID)
+				{
+					continue;
+				}
+				const TArray<FVertexInstanceID>& TestVertexInstances = DuplicateMeshDescription.GetPolygonPerimeterVertexInstances(TestPolygonID);
+				if (TestVertexInstances.Num() != RefVertexInstances.Num())
+				{
+					continue;
+				}
+				bool bIdentical = true;
+				int32 UVIndex = 0;
+				for (FVertexInstanceID TestVertexInstanceID : TestVertexInstances)
+				{
+					if (VertexInstanceUVs[TestVertexInstanceID] != RefUVs[UVIndex])
+					{
+						bIdentical = false;
+						break;
+					}
+					UVIndex++;
+				}
+				if (bIdentical)
+				{
+					MatchPolygonID = TestPolygonID;
+					break;
+				}
+			}
+
+			if (MatchPolygonID == FPolygonID::Invalid)
+			{
+				UniquePolygons.Add(RefPolygonID);
+				for (FVertexInstanceID RefVertexInstanceID : RefVertexInstances)
+				{
+					RemapVertexInstance.Add(RefVertexInstanceID, RefVertexInstanceID);
+				}
+			}
+			else
+			{
+				const TArray<FVertexInstanceID>& TestVertexInstances = DuplicateMeshDescription.GetPolygonPerimeterVertexInstances(MatchPolygonID);
+				int32 VertexInstanceIndex = 0;
+				for (FVertexInstanceID RefVertexInstanceID : RefVertexInstances)
+				{
+					RemapVertexInstance.Add(RefVertexInstanceID, TestVertexInstances[VertexInstanceIndex]);
+					VertexInstanceIndex++;
+				}
+				ToDeletePolygons.Add(RefPolygonID);
+			}
+		}
+
+		//Delete polygons
+		if(ToDeletePolygons.Num() > 0)
+		{
+			TArray<FEdgeID> OrphanedEdges;
+			TArray<FVertexInstanceID> OrphanedVertexInstances;
+			TArray<FPolygonGroupID> OrphanedPolygonGroups;
+			TArray<FVertexID> OrphanedVertices;
+			for (FPolygonID PolygonID : ToDeletePolygons)
+			{
+				DuplicateMeshDescription.DeletePolygon(PolygonID, &OrphanedEdges, &OrphanedVertexInstances, &OrphanedPolygonGroups);
+			}
+			for (FPolygonGroupID PolygonGroupID : OrphanedPolygonGroups)
+			{
+				DuplicateMeshDescription.DeletePolygonGroup(PolygonGroupID);
+			}
+			for (FVertexInstanceID VertexInstanceID : OrphanedVertexInstances)
+			{
+				DuplicateMeshDescription.DeleteVertexInstance(VertexInstanceID, &OrphanedVertices);
+			}
+			for (FEdgeID EdgeID : OrphanedEdges)
+			{
+				DuplicateMeshDescription.DeleteEdge(EdgeID, &OrphanedVertices);
+			}
+			for (FVertexID VertexID : OrphanedVertices)
+			{
+				DuplicateMeshDescription.DeleteVertex(VertexID);
+			}
+			//Compact and Remap IDs so we have clean ID from 0 to n since we just erase some polygons
+			FElementIDRemappings RemappingInfos;
+			DuplicateMeshDescription.Compact(RemappingInfos);
+		}
+	}
 	// Find overlapping corners for UV generator. Allow some threshold - this should not produce any error in a case if resulting
 	// mesh will not merge these vertices.
 	FOverlappingCorners OverlappingCorners;
@@ -1245,14 +1508,18 @@
 	if (bPackSuccess)
 	{
 		Packer.CommitPackedUVs();
-		TVertexInstanceAttributesRef<FVector2D> VertexInstanceUVs = DuplicateMeshDescription.VertexInstanceAttributes().GetAttributesRef<FVector2D>(MeshAttribute::VertexInstance::TextureCoordinate);
+		TVertexInstanceAttributesConstRef<FVector2D> DupVertexInstanceUVs = DuplicateMeshDescription.VertexInstanceAttributes().GetAttributesRef<FVector2D>(MeshAttribute::VertexInstance::TextureCoordinate);
+		TVertexInstanceAttributesConstRef<FVector2D> VertexInstanceUVs = MeshDescription.VertexInstanceAttributes().GetAttributesRef<FVector2D>(MeshAttribute::VertexInstance::TextureCoordinate);
 		// Save generated UVs
 		check(VertexInstanceUVs.GetNumIndices() > 1);
 		OutTexCoords.AddZeroed(VertexInstanceUVs.GetNumElements());
 		int32 TextureCoordIndex = 0;
-		for (const FVertexInstanceID& VertexInstanceID : DuplicateMeshDescription.VertexInstances().GetElementIDs())
+		for (const FVertexInstanceID& VertexInstanceID : MeshDescription.VertexInstances().GetElementIDs())
 		{
-			OutTexCoords[TextureCoordIndex++] = VertexInstanceUVs.Get(VertexInstanceID, 1);	// UV1
+			FVertexInstanceID RemapID = bMergeIdenticalMaterials ? RemapVertexInstance[VertexInstanceID] : VertexInstanceID;
+			// Save generated UVs
+			OutTexCoords[TextureCoordIndex] = VertexInstanceUVs.Get(RemapID, 1);	// UV1
+			TextureCoordIndex++;
 		}
 	}
 
@@ -1310,49 +1577,6 @@
 	return true;
 }
 
-void FMeshDescriptionOperations::RemapPolygonGroups(FMeshDescription& MeshDescription, TMap<FPolygonGroupID, FPolygonGroupID>& Remap)
-{
-	TPolygonGroupAttributesRef<FName> PolygonGroupNames = MeshDescription.PolygonGroupAttributes().GetAttributesRef<FName>(MeshAttribute::PolygonGroup::ImportedMaterialSlotName);
-	
-	struct FOldPolygonGroupData
-	{
-		FName Name;
-		TArray<FPolygonID> Polygons;
-	};
-
-	TMap<FPolygonGroupID, FOldPolygonGroupData> OldData;
-	for (const FPolygonGroupID& PolygonGroupID : MeshDescription.PolygonGroups().GetElementIDs())
-	{
-		if (!Remap.Contains(PolygonGroupID) || PolygonGroupID == Remap[PolygonGroupID])
-		{
-			//No need to change this one
-			continue;
-		}
-		FOldPolygonGroupData& PolygonGroupData = OldData.FindOrAdd(PolygonGroupID);
-		PolygonGroupData.Name = PolygonGroupNames[PolygonGroupID];
-		FMeshPolygonGroup& PolygonGroup = MeshDescription.GetPolygonGroup(PolygonGroupID);
-		PolygonGroupData.Polygons = PolygonGroup.Polygons;
-		PolygonGroup.Polygons.Empty();
-		MeshDescription.DeletePolygonGroup(PolygonGroupID);
-	}
-	for (auto Kvp : OldData)
-	{
-		FPolygonGroupID GroupID = Kvp.Key;
-		FPolygonGroupID ToGroupID = Remap[GroupID];
-		if (!MeshDescription.PolygonGroups().IsValid(ToGroupID))
-		{
-			MeshDescription.CreatePolygonGroupWithID(ToGroupID);
-		}
-		TArray<FPolygonID>& Polygons = MeshDescription.GetPolygonGroup(ToGroupID).Polygons;
-		Polygons.Append(Kvp.Value.Polygons);
-		PolygonGroupNames[ToGroupID] = Kvp.Value.Name;
-		for (FPolygonID PolygonID : Polygons)
-		{
-			MeshDescription.GetPolygon(PolygonID).PolygonGroupID = ToGroupID;
-		}
-	}
-}
-
 void FMeshDescriptionOperations::GeneratePlanarUV(const FMeshDescription& MeshDescription, const FUVMapSettings& Settings, TArray<FVector2D>& OutTexCoords)
 {
 	FVector U = FVector::ForwardVector;
@@ -1604,4 +1828,93 @@
 	}
 }
 
+void FMeshDescriptionOperations::RemapPolygonGroups(FMeshDescription& MeshDescription, TMap<FPolygonGroupID, FPolygonGroupID>& Remap)
+{
+	TPolygonGroupAttributesRef<FName> PolygonGroupNames = MeshDescription.PolygonGroupAttributes().GetAttributesRef<FName>(MeshAttribute::PolygonGroup::ImportedMaterialSlotName);
+
+	struct FOldPolygonGroupData
+	{
+		FName Name;
+		TArray<FPolygonID> Polygons;
+	};
+
+	TMap<FPolygonGroupID, FOldPolygonGroupData> OldData;
+	for (const FPolygonGroupID& PolygonGroupID : MeshDescription.PolygonGroups().GetElementIDs())
+	{
+		if (!Remap.Contains(PolygonGroupID) || PolygonGroupID == Remap[PolygonGroupID])
+		{
+			//No need to change this one
+			continue;
+		}
+		FOldPolygonGroupData& PolygonGroupData = OldData.FindOrAdd(PolygonGroupID);
+		PolygonGroupData.Name = PolygonGroupNames[PolygonGroupID];
+		FMeshPolygonGroup& PolygonGroup = MeshDescription.GetPolygonGroup(PolygonGroupID);
+		PolygonGroupData.Polygons = PolygonGroup.Polygons;
+		PolygonGroup.Polygons.Empty();
+		MeshDescription.DeletePolygonGroup(PolygonGroupID);
+	}
+	for (auto Kvp : OldData)
+	{
+		FPolygonGroupID GroupID = Kvp.Key;
+		FPolygonGroupID ToGroupID = Remap[GroupID];
+		if (!MeshDescription.PolygonGroups().IsValid(ToGroupID))
+		{
+			MeshDescription.CreatePolygonGroupWithID(ToGroupID);
+		}
+		TArray<FPolygonID>& Polygons = MeshDescription.GetPolygonGroup(ToGroupID).Polygons;
+		Polygons.Append(Kvp.Value.Polygons);
+		PolygonGroupNames[ToGroupID] = Kvp.Value.Name;
+		for (FPolygonID PolygonID : Polygons)
+		{
+			MeshDescription.GetPolygon(PolygonID).PolygonGroupID = ToGroupID;
+		}
+	}
+}
+
+void FMeshDescriptionOperations::SwapPolygonPolygonGroup(FMeshDescription& MeshDescription, int32 SectionIndex, int32 TriangleIndexStart, int32 TriangleIndexEnd, bool bRemoveEmptyPolygonGroup)
+{
+	int32 TriangleIndex = 0;
+	TPolygonGroupAttributesRef<FName> PolygonGroupNames = MeshDescription.PolygonGroupAttributes().GetAttributesRef<FName>(MeshAttribute::PolygonGroup::ImportedMaterialSlotName);
+
+	FPolygonGroupID TargetPolygonGroupID(SectionIndex);
+	if (!bRemoveEmptyPolygonGroup)
+	{
+		while (!MeshDescription.PolygonGroups().IsValid(TargetPolygonGroupID))
+		{
+			TargetPolygonGroupID = MeshDescription.CreatePolygonGroup();
+			PolygonGroupNames[TargetPolygonGroupID] = FName(*(TEXT("SwapPolygonMaterialSlotName_") + FString::FromInt(TargetPolygonGroupID.GetValue())));
+			TargetPolygonGroupID = FPolygonGroupID(SectionIndex);
+		}
+	}
+	else
+	{
+		//This will not follow the SectionIndex value if the value is greater then the number of section (do not use this when merging meshes)
+		if (!MeshDescription.PolygonGroups().IsValid(TargetPolygonGroupID))
+		{
+			TargetPolygonGroupID = MeshDescription.CreatePolygonGroup();
+			PolygonGroupNames[TargetPolygonGroupID] = FName(*(TEXT("SwapPolygonMaterialSlotName_") + FString::FromInt(TargetPolygonGroupID.GetValue())));
+		}
+	}
+
+	for (const FPolygonID PolygonID : MeshDescription.Polygons().GetElementIDs())
+	{
+		int32 TriangleCount = MeshDescription.GetPolygonTriangles(PolygonID).Num();
+		if (TriangleIndex >= TriangleIndexStart && TriangleIndex < TriangleIndexEnd)
+		{
+			check(TriangleIndex + (TriangleCount-1) < TriangleIndexEnd);
+			FMeshPolygon& Polygon = MeshDescription.GetPolygon(PolygonID);
+			FPolygonGroupID OldpolygonGroupID = Polygon.PolygonGroupID;
+			if (OldpolygonGroupID != TargetPolygonGroupID)
+			{
+				MeshDescription.SetPolygonPolygonGroup(PolygonID, TargetPolygonGroupID);
+				if (bRemoveEmptyPolygonGroup && MeshDescription.GetPolygonGroupPolygons(OldpolygonGroupID).Num() < 1)
+				{
+					MeshDescription.DeletePolygonGroup(OldpolygonGroupID);
+				}
+			}
+		}
+		TriangleIndex += TriangleCount;
+	}
+}
+
 #undef LOCTEXT_NAMESPACE
diff --git a/Engine/Source/Developer/MeshDescriptionOperations/Public/MeshDescriptionOperations.h b/Engine/Source/Developer/MeshDescriptionOperations/Public/MeshDescriptionOperations.h
--- a/Engine/Source/Developer/MeshDescriptionOperations/Public/MeshDescriptionOperations.h
+++ b/Engine/Source/Developer/MeshDescriptionOperations/Public/MeshDescriptionOperations.h
@@ -13,6 +13,10 @@ struct FPolygonGroupID;
 
 DECLARE_LOG_CATEGORY_EXTERN(LogMeshDescriptionOperations, Log, All);
 
+typedef TMap<FPolygonGroupID, FPolygonGroupID> PolygonGroupMap;
+
+DECLARE_DELEGATE_ThreeParams(FAppendPolygonGroupsDelegate, const FMeshDescription& /*SourceMesh*/, FMeshDescription& /*TargetMesh*/, PolygonGroupMap& /*RemapPolygonGroup*/)
+
 //////////////////////////////////////////////////////////////////////////
 // Any operations on the mesh description that do not depend on engine module
 // should be implement here.
@@ -34,6 +38,19 @@ public:
 	/** Convert old FRawMesh format to MeshDescription. */
 	static void ConvertFromRawMesh(const FRawMesh& SourceRawMesh, FMeshDescription& DestinationMeshDescription, const TMap<int32, FName>& MaterialMap);
 
+	struct FAppendSettings
+	{
+		FAppendSettings()
+			: bMergeVertexColor(true)
+			, MergedAssetPivot(0.0f, 0.0f, 0.0f)
+		{}
+		FAppendPolygonGroupsDelegate PolygonGroupsDelegate;
+		bool bMergeVertexColor;
+		FVector MergedAssetPivot;
+	};
+
+	static void AppendMeshDescription(const FMeshDescription& SourceMesh, FMeshDescription& TargetMesh, FAppendSettings& AppendSettings);
+
 	/*
 	 * Check if all normals and tangents are valid, if not recompute them
 	 */
@@ -62,7 +79,7 @@ public:
 		const FOverlappingCorners& OverlappingCorners);
 
 	/** Create some UVs from the specified mesh description data. */
-	static bool GenerateUniqueUVsForStaticMesh(const FMeshDescription& MeshDescription, int32 TextureResolution, TArray<FVector2D>& OutTexCoords);
+	static bool GenerateUniqueUVsForStaticMesh(const FMeshDescription& MeshDescription, int32 TextureResolution, bool bMergeIdenticalMaterials, TArray<FVector2D>& OutTexCoords);
 
 	/** Add a UV channel to the MeshDescription. */
 	static bool AddUVChannel(FMeshDescription& MeshDescription);
@@ -82,9 +99,22 @@ public:
 	/** Generate box UV mapping for the MeshDescription */
 	static void GenerateBoxUV(const FMeshDescription& MeshDescription, const FUVMapSettings& Settings, TArray<FVector2D>& OutTexCoords);
 
+	static void RemapPolygonGroups(FMeshDescription& MeshDescription, TMap<FPolygonGroupID, FPolygonGroupID>& Remap);
+
+	/*
+	 * Move some polygon to a new PolygonGroup(section)
+	 * SectionIndex: The target section we want to assign the polygon. See bRemoveEmptyPolygonGroup to know how its used
+	 * TriangleIndexStart: The triangle index is compute as follow: foreach polygon {TriangleIndex += Polygon->NumberTriangles}
+	 * TriangleIndexEnd: The triangle index is compute as follow: foreach polygon {TriangleIndex += Polygon->NumberTriangles}
+	 * bRemoveEmptyPolygonGroup: If true, any polygonGroup that is empty after moving a polygon will be delete.
+	 *                           This parameter impact how SectionIndex is use
+	 *                           If param is true  : PolygonGroupTargetID.GetValue() do not necessary equal SectionIndex in case there is less sections then SectionIndex
+	 *                           If param is false : PolygonGroupTargetID.GetValue() equal SectionIndex, we will add all necessary missing PolygonGroupID (this can generate empty PolygonGroupID)
+	 */
+	static void SwapPolygonPolygonGroup(FMeshDescription& MeshDescription, int32 SectionIndex, int32 TriangleIndexStart, int32 TriangleIndexEnd, bool bRemoveEmptyPolygonGroup);
+
 	static void ConvertHardEdgesToSmoothGroup(const FMeshDescription& SourceMeshDescription, TArray<uint32>& FaceSmoothingMasks);
 
 	static void ConvertSmoothGroupToHardEdges(const TArray<uint32>& FaceSmoothingMasks, FMeshDescription& DestinationMeshDescription);
 
-	static void RemapPolygonGroups(FMeshDescription& MeshDescription, TMap<FPolygonGroupID, FPolygonGroupID>& Remap);
 };
\ No newline at end of file
diff --git a/Engine/Source/Developer/MeshMergeUtilities/Private/MeshMergeDataTracker.cpp b/Engine/Source/Developer/MeshMergeUtilities/Private/MeshMergeDataTracker.cpp
--- a/Engine/Source/Developer/MeshMergeUtilities/Private/MeshMergeDataTracker.cpp
+++ b/Engine/Source/Developer/MeshMergeUtilities/Private/MeshMergeDataTracker.cpp
@@ -3,6 +3,7 @@
 #include "MeshMergeDataTracker.h"
 #include "MeshMergeHelpers.h"
 #include "Misc/Crc.h"
+#include "Engine/StaticMesh.h"
 
 FMeshMergeDataTracker::FMeshMergeDataTracker()
 	: AvailableLightMapUVChannel(INDEX_NONE), SummedLightMapPixels(0)
@@ -11,10 +12,12 @@ FMeshMergeDataTracker::FMeshMergeDataTracker()
 	FMemory::Memzero(bOcuppiedUVChannels);
 }
 
-FRawMesh& FMeshMergeDataTracker::AddAndRetrieveRawMesh(int32 MeshIndex, int32 LODIndex, UStaticMesh* InMesh)
+FMeshDescription& FMeshMergeDataTracker::AddAndRetrieveRawMesh(int32 MeshIndex, int32 LODIndex, UStaticMesh* InMesh)
 {
 	checkf(!RawMeshLODs.Contains(FMeshLODKey(MeshIndex, LODIndex, InMesh)), TEXT("Raw Mesh already added for this key"));
-	return RawMeshLODs.Add(FMeshLODKey(MeshIndex, LODIndex, InMesh));
+	FMeshDescription& MeshDescription = RawMeshLODs.Add(FMeshLODKey(MeshIndex, LODIndex, InMesh));
+	UStaticMesh::RegisterMeshAttributes(MeshDescription);
+	return MeshDescription;
 }
 
 void FMeshMergeDataTracker::RemoveRawMesh(int32 MeshIndex, int32 LODIndex)
@@ -128,13 +131,13 @@ bool FMeshMergeDataTracker::DoesAnyLODContainVertexColors() const
 bool FMeshMergeDataTracker::DoesUVChannelContainData(int32 UVChannel, int32 LODIndex) const
 {
 	checkf(FMath::IsWithinInclusive(LODIndex, 0, MAX_STATIC_MESH_LODS - 1), TEXT("Invalid LOD index"));
-	checkf(FMath::IsWithinInclusive(UVChannel, 0, MAX_MESH_TEXTURE_COORDS - 1), TEXT("Invalid UV channel index"));
+	checkf(FMath::IsWithinInclusive(UVChannel, 0, MAX_MESH_TEXTURE_COORDS_MD - 1), TEXT("Invalid UV channel index"));
 	return bOcuppiedUVChannels[LODIndex][UVChannel];
 }
 
 bool FMeshMergeDataTracker::DoesUVChannelContainData(int32 UVChannel) const
 {
-	checkf(FMath::IsWithinInclusive(UVChannel, 0, MAX_MESH_TEXTURE_COORDS - 1), TEXT("Invalid UV channel index"));
+	checkf(FMath::IsWithinInclusive(UVChannel, 0, MAX_MESH_TEXTURE_COORDS_MD - 1), TEXT("Invalid UV channel index"));
 	for(int32 LODIndex = 0; LODIndex < MAX_STATIC_MESH_LODS; LODIndex++)
 	{
 		if(bOcuppiedUVChannels[LODIndex][UVChannel])
@@ -167,21 +170,21 @@ uint32 FMeshMergeDataTracker::GetComponentToWedgeMappng(int32 MeshIndex, int32 L
 	return MappingPtr ? *MappingPtr : INDEX_NONE;
 }
 
-FRawMesh* FMeshMergeDataTracker::GetRawMeshPtr(int32 MeshIndex, int32 LODIndex)
+FMeshDescription* FMeshMergeDataTracker::GetRawMeshPtr(int32 MeshIndex, int32 LODIndex)
 {
 	return RawMeshLODs.Find(FMeshLODKey(MeshIndex, LODIndex));
 }
 
-FRawMesh* FMeshMergeDataTracker::GetRawMeshPtr(FMeshLODKey Key)
+FMeshDescription* FMeshMergeDataTracker::GetRawMeshPtr(FMeshLODKey Key)
 {
 	return RawMeshLODs.Find(Key);
 }
 
-FRawMesh* FMeshMergeDataTracker::FindRawMeshAndLODIndex(int32 MeshIndex, int32& OutLODIndex)
+FMeshDescription* FMeshMergeDataTracker::FindRawMeshAndLODIndex(int32 MeshIndex, int32& OutLODIndex)
 {
-	FRawMesh* FoundMeshPtr = nullptr;
+	FMeshDescription* FoundMeshPtr = nullptr;
 	OutLODIndex = INDEX_NONE;
-	for (TPair<FMeshLODKey, FRawMesh>& Pair : RawMeshLODs)
+	for (TPair<FMeshLODKey, FMeshDescription>& Pair : RawMeshLODs)
 	{
 		if (Pair.Key.GetMeshIndex() == MeshIndex)
 		{
@@ -194,13 +197,13 @@ FRawMesh* FMeshMergeDataTracker::FindRawMeshAndLODIndex(int32 MeshIndex, int32&
 	return FoundMeshPtr;
 }
 
-FRawMesh* FMeshMergeDataTracker::TryFindRawMeshForLOD(int32 MeshIndex, int32& InOutDesiredLODIndex)
+FMeshDescription* FMeshMergeDataTracker::TryFindRawMeshForLOD(int32 MeshIndex, int32& InOutDesiredLODIndex)
 {
-	FRawMesh* FoundMeshPtr = RawMeshLODs.Find(FMeshLODKey(MeshIndex, InOutDesiredLODIndex));
+	FMeshDescription* FoundMeshPtr = RawMeshLODs.Find(FMeshLODKey(MeshIndex, InOutDesiredLODIndex));
 	int32 SearchIndex = InOutDesiredLODIndex - 1;
 	while (FoundMeshPtr == nullptr && SearchIndex >= 0)
 	{
-		for (TPair<FMeshLODKey, FRawMesh>& Pair : RawMeshLODs)
+		for (TPair<FMeshLODKey, FMeshDescription>& Pair : RawMeshLODs)
 		{
 			if (Pair.Key.GetMeshIndex() == MeshIndex && Pair.Key.GetLODIndex() == SearchIndex)
 			{
@@ -234,35 +237,38 @@ void FMeshMergeDataTracker::GetMappingsForMeshLOD(FMeshLODKey Key, TArray<Sectio
 
 void FMeshMergeDataTracker::ProcessRawMeshes()
 {
-	bool bPotentialLightmapUVChannels[MAX_MESH_TEXTURE_COORDS];
+	bool bPotentialLightmapUVChannels[MAX_MESH_TEXTURE_COORDS_MD];
 	FMemory::Memset(bPotentialLightmapUVChannels, 1);
-	bool bPotentialLODLightmapUVChannels[MAX_STATIC_MESH_LODS][MAX_MESH_TEXTURE_COORDS];
+	bool bPotentialLODLightmapUVChannels[MAX_STATIC_MESH_LODS][MAX_MESH_TEXTURE_COORDS_MD];
 	FMemory::Memset(bPotentialLODLightmapUVChannels, 1);
 
 	// Retrieve information in regards to occupied UV channels whether or not a mesh contains vertex colors, and if 
-	for (TPair<FMeshLODKey, FRawMesh>& MeshPair : RawMeshLODs)
+	for (TPair<FMeshLODKey, FMeshDescription>& MeshPair : RawMeshLODs)
 	{
 		FMeshLODKey& Key = MeshPair.Key;
 		const int32 LODIndex = Key.GetLODIndex();
-		const FRawMesh& RawMesh = MeshPair.Value;
+		const FMeshDescription& RawMesh = MeshPair.Value;
+
+		TVertexInstanceAttributesConstRef<FVector4> VertexInstanceColors = RawMesh.VertexInstanceAttributes().GetAttributesRef<FVector4>(MeshAttribute::VertexInstance::Color);
+		TVertexInstanceAttributesConstRef<FVector2D> VertexInstanceUVs = RawMesh.VertexInstanceAttributes().GetAttributesRef<FVector2D>(MeshAttribute::VertexInstance::TextureCoordinate);
 
 		// hash vertex color buffer so we can see if instances have unique vertex data
-		if(RawMesh.WedgeColors.Num() > 0)
+		if(VertexInstanceColors.GetNumElements() > 0)
 		{
-			Key.SetVertexColorHash(FCrc::MemCrc32(RawMesh.WedgeColors.GetData(), RawMesh.WedgeColors.Num() * sizeof(FColor)));
+			Key.SetVertexColorHash(RawMesh.VertexInstanceAttributes().GetHash(MeshAttribute::VertexInstance::Color));
 		}
 
 		const int32 LightmapChannelIdx = LightmapChannelLODs.FindRef(Key);
 		bool bNeedsVertexData = false;
 		
-		for (int32 ChannelIndex = 0; ChannelIndex < MAX_MESH_TEXTURE_COORDS; ++ChannelIndex)
+		if (VertexInstanceUVs.GetNumElements() > 0)
 		{
-			if (RawMesh.WedgeTexCoords[ChannelIndex].Num())
+			for (int32 ChannelIndex = 0; ChannelIndex < FMath::Min(VertexInstanceUVs.GetNumIndices(), (int32)MAX_MESH_TEXTURE_COORDS_MD); ++ChannelIndex)
 			{
 				bOcuppiedUVChannels[LODIndex][ChannelIndex] = true;
 				bPotentialLODLightmapUVChannels[LODIndex][ChannelIndex] = (ChannelIndex == LightmapChannelIdx);
-
-				const bool bWrappingUVs = FMeshMergeHelpers::CheckWrappingUVs(RawMesh.WedgeTexCoords[ChannelIndex]);
+				
+				const bool bWrappingUVs = FMeshMergeHelpers::CheckWrappingUVs(RawMesh, ChannelIndex);
 				if (bWrappingUVs)
 				{
 					bNeedsVertexData = true;
@@ -271,7 +277,7 @@ void FMeshMergeDataTracker::ProcessRawMeshes()
 		}
 
 		// Merge available lightmap slots from LODs into one set, so we can assess later what slots are available
-		for (int32 ChannelIdx = 1; ChannelIdx < MAX_MESH_TEXTURE_COORDS; ++ChannelIdx)
+		for (int32 ChannelIdx = 1; ChannelIdx < MAX_MESH_TEXTURE_COORDS_MD; ++ChannelIdx)
 		{
 			bPotentialLightmapUVChannels[ChannelIdx] &= bPotentialLODLightmapUVChannels[LODIndex][ChannelIdx];
 		}
@@ -281,13 +287,13 @@ void FMeshMergeDataTracker::ProcessRawMeshes()
 			RequiresUniqueUVs.Add(Key);
 		}
 
-		bWithVertexColors[LODIndex] |= RawMesh.WedgeColors.Num() != 0;
+		bWithVertexColors[LODIndex] |= VertexInstanceColors.GetNumElements() != 0;
 	}
 
 	// Look for an available lightmap slot we can use in the merged set
 	// We start at channel 1 as merged meshes always use texcoord 0 for their expected mapping channel, so we cant use it;
 	AvailableLightMapUVChannel = INDEX_NONE;
-	for (int32 ChannelIdx = 1; ChannelIdx < MAX_MESH_TEXTURE_COORDS; ++ChannelIdx)
+	for (int32 ChannelIdx = 1; ChannelIdx < MAX_MESH_TEXTURE_COORDS_MD; ++ChannelIdx)
 	{
 		if(bPotentialLightmapUVChannels[ChannelIdx])
 		{
diff --git a/Engine/Source/Developer/MeshMergeUtilities/Private/MeshMergeHelpers.cpp b/Engine/Source/Developer/MeshMergeUtilities/Private/MeshMergeHelpers.cpp
--- a/Engine/Source/Developer/MeshMergeUtilities/Private/MeshMergeHelpers.cpp
+++ b/Engine/Source/Developer/MeshMergeUtilities/Private/MeshMergeHelpers.cpp
@@ -6,8 +6,6 @@
 #include "Engine/MeshMerging.h"
 
 #include "MaterialOptions.h"
-#include "RawMesh.h"
-
 #include "MeshDescription.h"
 #include "MeshAttributes.h"
 #include "MeshAttributeArray.h"
@@ -177,13 +175,13 @@ void FMeshMergeHelpers::ExtractSections(const UStaticMesh* StaticMesh, int32 LOD
 	}
 }
 
-void FMeshMergeHelpers::ExpandInstances(const UInstancedStaticMeshComponent* InInstancedStaticMeshComponent, FRawMesh& InOutRawMesh, TArray<FSectionInfo>& InOutSections)
+void FMeshMergeHelpers::ExpandInstances(const UInstancedStaticMeshComponent* InInstancedStaticMeshComponent, FMeshDescription& InOutRawMesh, TArray<FSectionInfo>& InOutSections)
 {
-	FRawMesh CombinedRawMesh;
+	FMeshDescription CombinedRawMesh;
 
 	for(const FInstancedStaticMeshInstanceData& InstanceData : InInstancedStaticMeshComponent->PerInstanceSMData)
 	{
-		FRawMesh InstanceRawMesh = InOutRawMesh;
+		FMeshDescription InstanceRawMesh = InOutRawMesh;
 		FMeshMergeHelpers::TransformRawMeshVertexData(FTransform(InstanceData.Transform), InstanceRawMesh);
 		FMeshMergeHelpers::AppendRawMesh(CombinedRawMesh, InstanceRawMesh);
 	}
@@ -191,7 +189,7 @@ void FMeshMergeHelpers::ExpandInstances(const UInstancedStaticMeshComponent* InI
 	InOutRawMesh = CombinedRawMesh;
 }
 
-void FMeshMergeHelpers::RetrieveMesh(const UStaticMeshComponent* StaticMeshComponent, int32 LODIndex, FRawMesh& RawMesh, bool bPropagateVertexColours)
+void FMeshMergeHelpers::RetrieveMesh(const UStaticMeshComponent* StaticMeshComponent, int32 LODIndex, FMeshDescription& RawMesh, bool bPropagateVertexColours)
 {
 	const UStaticMesh* StaticMesh = StaticMeshComponent->GetStaticMesh();
 	const FStaticMeshSourceModel& StaticMeshModel = StaticMesh->SourceModels[LODIndex];
@@ -199,13 +197,13 @@ void FMeshMergeHelpers::RetrieveMesh(const UStaticMeshComponent* StaticMeshCompo
 	const bool bIsSplineMeshComponent = StaticMeshComponent->IsA<USplineMeshComponent>();
 
 	// Imported meshes will have a filled RawMeshBulkData set
-	const bool bImportedMesh = !StaticMeshModel.RawMeshBulkData->IsEmpty();
+	const bool bImportedMesh = !StaticMeshModel.IsRawMeshEmpty();
 		
 	// Export the raw mesh data using static mesh render data
-	ExportStaticMeshLOD(StaticMesh->RenderData->LODResources[LODIndex], RawMesh);	
+	ExportStaticMeshLOD(StaticMesh->RenderData->LODResources[LODIndex], RawMesh, StaticMesh->StaticMaterials);
 
 	// Make sure the raw mesh is not irreparably malformed.
-	if (!RawMesh.IsValid())
+	if (RawMesh.VertexInstances().Num() <= 0)
 	{
 		return;
 	}
@@ -233,23 +231,22 @@ void FMeshMergeHelpers::RetrieveMesh(const UStaticMeshComponent* StaticMeshCompo
 	// Transform raw mesh vertex data by the Static Mesh Component's component to world transformation	
 	TransformRawMeshVertexData(ComponentToWorldTransform, RawMesh);
 
-	if (!RawMesh.IsValid())
+	if (RawMesh.VertexInstances().Num() <= 0)
 	{
 		return;
 	}
 
 	// Figure out if we should recompute normals and tangents. By default generated LODs should not recompute normals	
-	const bool bRecomputeNormals = RawMesh.WedgeTangentZ.Num() == 0;
-	const bool bRecomputeTangents = RawMesh.WedgeTangentX.Num() == 0 || RawMesh.WedgeTangentY.Num() == 0;
-
-	if (bRecomputeNormals || bRecomputeTangents)
+	uint32 TangentOptions = FMeshDescriptionOperations::ETangentOptions::BlendOverlappingNormals;
+	if (BuildSettings.bRemoveDegenerates)
 	{
-		IMeshUtilities& Utilities = FModuleManager::Get().LoadModuleChecked<IMeshUtilities>("MeshUtilities");
-		Utilities.RecomputeTangentsAndNormalsForRawMesh(bRecomputeTangents, bRecomputeNormals, BuildSettings, RawMesh);
+		// If removing degenerate triangles, ignore them when computing tangents.
+		TangentOptions |= FMeshDescriptionOperations::ETangentOptions::IgnoreDegenerateTriangles;
 	}
+	FMeshDescriptionOperations::RecomputeNormalsAndTangentsIfNeeded(RawMesh, (FMeshDescriptionOperations::ETangentOptions)TangentOptions, BuildSettings.bUseMikkTSpace);
 }
 
-void FMeshMergeHelpers::RetrieveMesh(USkeletalMeshComponent* SkeletalMeshComponent, int32 LODIndex, FRawMesh& RawMesh, bool bPropagateVertexColours)
+void FMeshMergeHelpers::RetrieveMesh(USkeletalMeshComponent* SkeletalMeshComponent, int32 LODIndex, FMeshDescription& RawMesh, bool bPropagateVertexColours)
 {
 	FSkeletalMeshModel* Resource = SkeletalMeshComponent->SkeletalMesh->GetImportedModel();
 	if (Resource->LODModels.IsValidIndex(LODIndex))
@@ -261,83 +258,142 @@ void FMeshMergeHelpers::RetrieveMesh(USkeletalMeshComponent* SkeletalMeshCompone
 		SkeletalMeshComponent->GetCPUSkinnedVertices(FinalVertices, LODIndex);
 
 		FSkeletalMeshLODModel& LODModel = Resource->LODModels[LODIndex];
+		
+		const int32 NumSections = LODModel.Sections.Num();
+		
+		//Empty the raw mesh
+		RawMesh.Empty();
+
+		TVertexAttributesRef<FVector> VertexPositions = RawMesh.VertexAttributes().GetAttributesRef<FVector>(MeshAttribute::Vertex::Position);
+		TEdgeAttributesRef<bool> EdgeHardnesses = RawMesh.EdgeAttributes().GetAttributesRef<bool>(MeshAttribute::Edge::IsHard);
+		TEdgeAttributesRef<float> EdgeCreaseSharpnesses = RawMesh.EdgeAttributes().GetAttributesRef<float>(MeshAttribute::Edge::CreaseSharpness);
+		TPolygonGroupAttributesRef<FName> PolygonGroupImportedMaterialSlotNames = RawMesh.PolygonGroupAttributes().GetAttributesRef<FName>(MeshAttribute::PolygonGroup::ImportedMaterialSlotName);
+		TVertexInstanceAttributesRef<FVector> VertexInstanceNormals = RawMesh.VertexInstanceAttributes().GetAttributesRef<FVector>(MeshAttribute::VertexInstance::Normal);
+		TVertexInstanceAttributesRef<FVector> VertexInstanceTangents = RawMesh.VertexInstanceAttributes().GetAttributesRef<FVector>(MeshAttribute::VertexInstance::Tangent);
+		TVertexInstanceAttributesRef<float> VertexInstanceBinormalSigns = RawMesh.VertexInstanceAttributes().GetAttributesRef<float>(MeshAttribute::VertexInstance::BinormalSign);
+		TVertexInstanceAttributesRef<FVector4> VertexInstanceColors = RawMesh.VertexInstanceAttributes().GetAttributesRef<FVector4>(MeshAttribute::VertexInstance::Color);
+		TVertexInstanceAttributesRef<FVector2D> VertexInstanceUVs = RawMesh.VertexInstanceAttributes().GetAttributesRef<FVector2D>(MeshAttribute::VertexInstance::TextureCoordinate);
+
+		int32 TotalTriangles = 0;
+		int32 TotalCorners = 0;
+		for (int32 SectionIndex = 0; SectionIndex < NumSections; SectionIndex++)
+		{
+			const FSkelMeshSection& SkelMeshSection = LODModel.Sections[SectionIndex];
+			TotalTriangles += SkelMeshSection.NumTriangles;
+		}
+		TotalCorners = TotalTriangles * 3;
+		RawMesh.ReserveNewVertices(FinalVertices.Num());
+		RawMesh.ReserveNewPolygons(TotalTriangles);
+		RawMesh.ReserveNewVertexInstances(TotalCorners);
+		RawMesh.ReserveNewEdges(TotalCorners);
 
 		// Copy skinned vertex positions
 		for (int32 VertIndex = 0; VertIndex < FinalVertices.Num(); ++VertIndex)
 		{
-			RawMesh.VertexPositions.Add(FinalVertices[VertIndex].Position);
+			const FVertexID VertexID = RawMesh.CreateVertex();
+			VertexPositions[VertexID] = FinalVertices[VertIndex].Position;
 		}
 
-		const int32 NumSections = LODModel.Sections.Num();
+		VertexInstanceUVs.SetNumIndices(MAX_TEXCOORDS);
 
+		
 		for (int32 SectionIndex = 0; SectionIndex < NumSections; SectionIndex++)
 		{
 			const FSkelMeshSection& SkelMeshSection = LODModel.Sections[SectionIndex];
-			// Build 'wedge' info
 			const int32 NumWedges = SkelMeshSection.NumTriangles * 3;
-			for (int32 WedgeIndex = 0; WedgeIndex < NumWedges; WedgeIndex++)
-			{
-				const int32 VertexIndexForWedge = LODModel.IndexBuffer[SkelMeshSection.BaseIndex + WedgeIndex];
 
-				RawMesh.WedgeIndices.Add(VertexIndexForWedge);
-
-				const FSoftSkinVertex& SoftVertex = SkelMeshSection.SoftVertices[VertexIndexForWedge - SkelMeshSection.BaseVertexIndex];
-
-				const FFinalSkinVertex& SkinnedVertex = FinalVertices[VertexIndexForWedge];
-				const FVector TangentX = SkinnedVertex.TangentX.ToFVector();
-				const FVector TangentZ = SkinnedVertex.TangentZ.ToFVector();
-				const FVector4 UnpackedTangentZ = SkinnedVertex.TangentZ.ToFVector4();
-				const FVector TangentY = (TangentX ^ TangentZ).GetSafeNormal() * UnpackedTangentZ.W;
-
-				RawMesh.WedgeTangentX.Add(TangentX);
-				RawMesh.WedgeTangentY.Add(TangentY);
-				RawMesh.WedgeTangentZ.Add(TangentZ);
+			//Create the polygon group ID
+			int32 SectionMaterialIndex = SkelMeshSection.MaterialIndex;
+			int32 MaterialIndex = SectionMaterialIndex;
+			// use the remapping of material indices for all LODs besides the base LOD 
+			if (LODIndex > 0 && SrcLODInfo.LODMaterialMap.IsValidIndex(SkelMeshSection.MaterialIndex))
+			{
+				MaterialIndex = FMath::Clamp<int32>(SrcLODInfo.LODMaterialMap[SkelMeshSection.MaterialIndex], 0, SkeletalMeshComponent->SkeletalMesh->Materials.Num());
+			}
 
-				for (uint32 TexCoordIndex = 0; TexCoordIndex < MAX_MESH_TEXTURE_COORDS; TexCoordIndex++)
+			FName ImportedMaterialSlotName = SkeletalMeshComponent->SkeletalMesh->Materials[MaterialIndex].ImportedMaterialSlotName;
+			const FPolygonGroupID SectionPolygonGroupID(SectionMaterialIndex);
+			if (!RawMesh.IsPolygonGroupValid(SectionPolygonGroupID))
+			{
+				RawMesh.CreatePolygonGroupWithID(SectionPolygonGroupID);
+				PolygonGroupImportedMaterialSlotNames[SectionPolygonGroupID] = ImportedMaterialSlotName;
+			}
+			int32 WedgeIndex = 0;
+			for (uint32 SectionTriangleIndex = 0; SectionTriangleIndex < SkelMeshSection.NumTriangles; ++SectionTriangleIndex)
+			{
+				FVertexID VertexIndexes[3];
+				FVertexInstanceID VertexInstanceIDs[3];
+				for (int32 CornerIndex = 0; CornerIndex < 3; ++CornerIndex, ++WedgeIndex)
 				{
-					if (TexCoordIndex >= MAX_TEXCOORDS)
-					{
-						RawMesh.WedgeTexCoords[TexCoordIndex].AddDefaulted();
-					}
-					else
+					const int32 VertexIndexForWedge = LODModel.IndexBuffer[SkelMeshSection.BaseIndex + WedgeIndex];
+					VertexIndexes[CornerIndex] = FVertexID(VertexIndexForWedge);
+					FVertexInstanceID VertexInstanceID = RawMesh.CreateVertexInstance(VertexIndexes[CornerIndex]);
+					VertexInstanceIDs[CornerIndex] = VertexInstanceID;
+					
+					const FSoftSkinVertex& SoftVertex = SkelMeshSection.SoftVertices[VertexIndexForWedge - SkelMeshSection.BaseVertexIndex];
+					const FFinalSkinVertex& SkinnedVertex = FinalVertices[VertexIndexForWedge];
+
+					//Set NTBs
+					const FVector TangentX = SkinnedVertex.TangentX.ToFVector();
+					const FVector TangentZ = SkinnedVertex.TangentZ.ToFVector();
+					//@todo: do we need to inverse the sign between skeletalmesh and staticmesh, the old code was doing so.
+					const float TangentYSign = SkinnedVertex.TangentZ.ToFVector4().W;
+					
+					VertexInstanceTangents[VertexInstanceID] = TangentX;
+					VertexInstanceBinormalSigns[VertexInstanceID] = TangentYSign;
+					VertexInstanceNormals[VertexInstanceID] = TangentZ;
+
+					for (uint32 TexCoordIndex = 0; TexCoordIndex < MAX_TEXCOORDS; TexCoordIndex++)
 					{
-						RawMesh.WedgeTexCoords[TexCoordIndex].Add(SoftVertex.UVs[TexCoordIndex]);
+						//Add this vertex instance tex coord
+						VertexInstanceUVs.Set(VertexInstanceID, TexCoordIndex, SoftVertex.UVs[TexCoordIndex]);
 					}
-				}
 
-				if (bPropagateVertexColours)
-				{
-					RawMesh.WedgeColors.Add(SoftVertex.Color);
+					//Add this vertex instance color
+					VertexInstanceColors[VertexInstanceID] = bPropagateVertexColours ? FVector4(FLinearColor(SoftVertex.Color)) : FVector4(1.0f, 1.0f, 1.0f);
 				}
-				else
+				//Create a polygon from this triangle
+				TArray<FMeshDescription::FContourPoint> Contours;
+				for (int32 CornerIndex = 0; CornerIndex < 3; ++CornerIndex)
 				{
-					RawMesh.WedgeColors.Add(FColor::White);
+					int32 ContourPointIndex = Contours.AddDefaulted();
+					FMeshDescription::FContourPoint& ContourPoint = Contours[ContourPointIndex];
+					//Find the matching edge ID
+					uint32 CornerIndices[2];
+					CornerIndices[0] = (CornerIndex + 0) % 3;
+					CornerIndices[1] = (CornerIndex + 1) % 3;
+
+					FVertexID EdgeVertexIDs[2];
+					EdgeVertexIDs[0] = VertexIndexes[CornerIndices[0]];
+					EdgeVertexIDs[1] = VertexIndexes[CornerIndices[1]];
+
+					FEdgeID MatchEdgeId = RawMesh.GetVertexPairEdge(EdgeVertexIDs[0], EdgeVertexIDs[1]);
+					if (MatchEdgeId == FEdgeID::Invalid)
+					{
+						MatchEdgeId = RawMesh.CreateEdge(EdgeVertexIDs[0], EdgeVertexIDs[1]);
+						EdgeHardnesses[MatchEdgeId] = false;
+						EdgeCreaseSharpnesses[MatchEdgeId] = 0.0f;
+					}
+					ContourPoint.EdgeID = MatchEdgeId;
+					ContourPoint.VertexInstanceID = VertexInstanceIDs[CornerIndices[0]];
 				}
-			}
-
-			int32 MaterialIndex = SkelMeshSection.MaterialIndex;
-			// use the remapping of material indices for all LODs besides the base LOD 
-			if (LODIndex > 0 && SrcLODInfo.LODMaterialMap.IsValidIndex(SkelMeshSection.MaterialIndex))
-			{
-				MaterialIndex = FMath::Clamp<int32>(SrcLODInfo.LODMaterialMap[SkelMeshSection.MaterialIndex], 0, SkeletalMeshComponent->SkeletalMesh->Materials.Num());
-			}
-
-			// copy face info
-			for (uint32 TriIndex = 0; TriIndex < SkelMeshSection.NumTriangles; TriIndex++)
-			{
-				RawMesh.FaceMaterialIndices.Add(MaterialIndex);
-				RawMesh.FaceSmoothingMasks.Add(0); // Assume this is ignored as bRecomputeNormals is false
+				// Insert a polygon into the mesh
+				const FPolygonID NewPolygonID = RawMesh.CreatePolygon(SectionPolygonGroupID, Contours);
+				//Triangulate the polygon
+				FMeshPolygon& Polygon = RawMesh.GetPolygon(NewPolygonID);
+				RawMesh.ComputePolygonTriangulation(NewPolygonID, Polygon.Triangles);
 			}
 		}
 	}
 }
 
-void FMeshMergeHelpers::RetrieveMesh(const UStaticMesh* StaticMesh, int32 LODIndex, FRawMesh& RawMesh)
+void FMeshMergeHelpers::RetrieveMesh(const UStaticMesh* StaticMesh, int32 LODIndex, FMeshDescription& RawMesh)
 {
 	const FStaticMeshSourceModel& StaticMeshModel = StaticMesh->SourceModels[LODIndex];
 
 	// Imported meshes will have a filled RawMeshBulkData set
-	const bool bImportedMesh = !StaticMeshModel.IsRawMeshEmpty();
+	const bool bImportedMesh = !StaticMeshModel.IsRawMeshEmpty() || (StaticMeshModel.OriginalMeshDescription != nullptr);
+	
 	// Check whether or not this mesh has been reduced in-engine
 	const bool bReducedMesh = (StaticMeshModel.ReductionSettings.PercentTriangles < 1.0f);
 	// Trying to retrieve rawmesh from SourceStaticMeshModel was giving issues, which causes a mismatch			
@@ -345,15 +401,15 @@ void FMeshMergeHelpers::RetrieveMesh(const UStaticMesh* StaticMesh, int32 LODInd
 
 	if (bImportedMesh && !bReducedMesh && !bRenderDataMismatch)
 	{
-		StaticMeshModel.LoadRawMesh(RawMesh);
+		StaticMesh->GetOriginalMeshDescription(LODIndex, RawMesh);
 	}
 	else
 	{
-		ExportStaticMeshLOD(StaticMesh->RenderData->LODResources[LODIndex], RawMesh);
+		ExportStaticMeshLOD(StaticMesh->RenderData->LODResources[LODIndex], RawMesh, StaticMesh->StaticMaterials);
 	}
 
 	// Make sure the raw mesh is not irreparably malformed.
-	if (!RawMesh.IsValid())
+	if (RawMesh.VertexInstances().Num() <= 0)
 	{
 		// wrong
 		bool check = true;
@@ -362,102 +418,151 @@ void FMeshMergeHelpers::RetrieveMesh(const UStaticMesh* StaticMesh, int32 LODInd
 	// Use build settings from base mesh for LOD entries that was generated inside Editor.
 	const FMeshBuildSettings& BuildSettings = bImportedMesh ? StaticMeshModel.BuildSettings : StaticMesh->SourceModels[0].BuildSettings;
 
-	// Figure out if we should recompute normals and tangents. By default generated LODs should not recompute normals
-	const bool bRecomputeNormals = (bImportedMesh && BuildSettings.bRecomputeNormals) || RawMesh.WedgeTangentZ.Num() == 0;
-	const bool bRecomputeTangents = (bImportedMesh && BuildSettings.bRecomputeTangents) || RawMesh.WedgeTangentX.Num() == 0 || RawMesh.WedgeTangentY.Num() == 0;
-
-	if (bRecomputeNormals || bRecomputeTangents)
+	// Figure out if we should recompute normals and tangents. By default generated LODs should not recompute normals	
+	uint32 TangentOptions = FMeshDescriptionOperations::ETangentOptions::BlendOverlappingNormals;
+	if (BuildSettings.bRemoveDegenerates)
 	{
-		IMeshUtilities& Utilities = FModuleManager::Get().LoadModuleChecked<IMeshUtilities>("MeshUtilities");
-		Utilities.RecomputeTangentsAndNormalsForRawMesh(bRecomputeTangents, bRecomputeNormals, BuildSettings, RawMesh);
+		// If removing degenerate triangles, ignore them when computing tangents.
+		TangentOptions |= FMeshDescriptionOperations::ETangentOptions::IgnoreDegenerateTriangles;
 	}
+	FMeshDescriptionOperations::RecomputeNormalsAndTangentsIfNeeded(RawMesh, (FMeshDescriptionOperations::ETangentOptions)TangentOptions, BuildSettings.bUseMikkTSpace, (bImportedMesh && BuildSettings.bRecomputeNormals), (bImportedMesh && BuildSettings.bRecomputeTangents));
 }
 
-void FMeshMergeHelpers::ExportStaticMeshLOD(const FStaticMeshLODResources& StaticMeshLOD, FRawMesh& OutRawMesh)
+void FMeshMergeHelpers::ExportStaticMeshLOD(const FStaticMeshLODResources& StaticMeshLOD, FMeshDescription& OutRawMesh, const TArray<FStaticMaterial>& Materials)
 {
 	const int32 NumWedges = StaticMeshLOD.IndexBuffer.GetNumIndices();
 	const int32 NumVertexPositions = StaticMeshLOD.VertexBuffers.PositionVertexBuffer.GetNumVertices();
 	const int32 NumFaces = NumWedges / 3;
 
-	// Indices
-	StaticMeshLOD.IndexBuffer.GetCopy(OutRawMesh.WedgeIndices);
+	OutRawMesh.Empty();
 
-	// Vertex positions
-	if (NumVertexPositions > 0)
+	if (NumVertexPositions <= 0 || StaticMeshLOD.VertexBuffers.StaticMeshVertexBuffer.GetNumVertices() <= 0)
 	{
-		OutRawMesh.VertexPositions.Empty(NumVertexPositions);
-		for (int32 PosIdx = 0; PosIdx < NumVertexPositions; ++PosIdx)
-		{
-			FVector Pos = StaticMeshLOD.VertexBuffers.PositionVertexBuffer.VertexPosition(PosIdx);
-			OutRawMesh.VertexPositions.Add(Pos);
-		}
+		return;
 	}
 
-	// Vertex data
-	if (StaticMeshLOD.VertexBuffers.StaticMeshVertexBuffer.GetNumVertices() > 0)
-	{
-		OutRawMesh.WedgeTangentX.Empty(NumWedges);
-		OutRawMesh.WedgeTangentY.Empty(NumWedges);
-		OutRawMesh.WedgeTangentZ.Empty(NumWedges);
+	TVertexAttributesRef<FVector> VertexPositions = OutRawMesh.VertexAttributes().GetAttributesRef<FVector>(MeshAttribute::Vertex::Position);
+	TEdgeAttributesRef<bool> EdgeHardnesses = OutRawMesh.EdgeAttributes().GetAttributesRef<bool>(MeshAttribute::Edge::IsHard);
+	TEdgeAttributesRef<float> EdgeCreaseSharpnesses = OutRawMesh.EdgeAttributes().GetAttributesRef<float>(MeshAttribute::Edge::CreaseSharpness);
+	TPolygonGroupAttributesRef<FName> PolygonGroupImportedMaterialSlotNames = OutRawMesh.PolygonGroupAttributes().GetAttributesRef<FName>(MeshAttribute::PolygonGroup::ImportedMaterialSlotName);
+	TVertexInstanceAttributesRef<FVector> VertexInstanceNormals = OutRawMesh.VertexInstanceAttributes().GetAttributesRef<FVector>(MeshAttribute::VertexInstance::Normal);
+	TVertexInstanceAttributesRef<FVector> VertexInstanceTangents = OutRawMesh.VertexInstanceAttributes().GetAttributesRef<FVector>(MeshAttribute::VertexInstance::Tangent);
+	TVertexInstanceAttributesRef<float> VertexInstanceBinormalSigns = OutRawMesh.VertexInstanceAttributes().GetAttributesRef<float>(MeshAttribute::VertexInstance::BinormalSign);
+	TVertexInstanceAttributesRef<FVector4> VertexInstanceColors = OutRawMesh.VertexInstanceAttributes().GetAttributesRef<FVector4>(MeshAttribute::VertexInstance::Color);
+	TVertexInstanceAttributesRef<FVector2D> VertexInstanceUVs = OutRawMesh.VertexInstanceAttributes().GetAttributesRef<FVector2D>(MeshAttribute::VertexInstance::TextureCoordinate);
+
+	OutRawMesh.ReserveNewVertices(NumVertexPositions);
+	OutRawMesh.ReserveNewVertexInstances(NumWedges);
+	OutRawMesh.ReserveNewPolygons(NumFaces);
+	OutRawMesh.ReserveNewEdges(NumWedges);
+
+	const int32 NumTexCoords = StaticMeshLOD.VertexBuffers.StaticMeshVertexBuffer.GetNumTexCoords();
+	VertexInstanceUVs.SetNumIndices(NumTexCoords);
 
-		const int32 NumTexCoords = StaticMeshLOD.VertexBuffers.StaticMeshVertexBuffer.GetNumTexCoords();
-		for (int32 TexCoodIdx = 0; TexCoodIdx < NumTexCoords; ++TexCoodIdx)
+	
+	for (int32 SectionIndex = 0; SectionIndex < StaticMeshLOD.Sections.Num(); ++SectionIndex)
+	{
+		const FStaticMeshSection& Section = StaticMeshLOD.Sections[SectionIndex];
+		FPolygonGroupID CurrentPolygonGroupID = OutRawMesh.CreatePolygonGroup();
+		check(CurrentPolygonGroupID.GetValue() == SectionIndex);
+		if (Materials.IsValidIndex(Section.MaterialIndex))
 		{
-			OutRawMesh.WedgeTexCoords[TexCoodIdx].Empty(NumWedges);
+			PolygonGroupImportedMaterialSlotNames[CurrentPolygonGroupID] = Materials[Section.MaterialIndex].ImportedMaterialSlotName;
 		}
-
-		for (int32 WedgeIndex : OutRawMesh.WedgeIndices)
+		else
 		{
-			FVector WedgeTangentX = StaticMeshLOD.VertexBuffers.StaticMeshVertexBuffer.VertexTangentX(WedgeIndex);
-			FVector WedgeTangentY = StaticMeshLOD.VertexBuffers.StaticMeshVertexBuffer.VertexTangentY(WedgeIndex);
-			FVector WedgeTangentZ = StaticMeshLOD.VertexBuffers.StaticMeshVertexBuffer.VertexTangentZ(WedgeIndex);
-			OutRawMesh.WedgeTangentX.Add(WedgeTangentX);
-			OutRawMesh.WedgeTangentY.Add(WedgeTangentY);
-			OutRawMesh.WedgeTangentZ.Add(WedgeTangentZ);
-
-			for (int32 TexCoodIdx = 0; TexCoodIdx < NumTexCoords; ++TexCoodIdx)
-			{
-				FVector2D WedgeTexCoord = StaticMeshLOD.VertexBuffers.StaticMeshVertexBuffer.GetVertexUV(WedgeIndex, TexCoodIdx);
-				OutRawMesh.WedgeTexCoords[TexCoodIdx].Add(WedgeTexCoord);
-			}
+			PolygonGroupImportedMaterialSlotNames[CurrentPolygonGroupID] = FName(*(TEXT("MeshMergeMaterial_") + FString::FromInt(SectionIndex)));
 		}
 	}
 
-	// Vertex colors
-	if (StaticMeshLOD.VertexBuffers.ColorVertexBuffer.GetNumVertices() > 0)
+	//Create the vertex
+	for (int32 VertexIndex = 0; VertexIndex < NumVertexPositions; ++VertexIndex)
 	{
-		OutRawMesh.WedgeColors.Empty(NumWedges);
-		for (int32 WedgeIndex : OutRawMesh.WedgeIndices)
-		{
-			FColor VertexColor = StaticMeshLOD.VertexBuffers.ColorVertexBuffer.VertexColor(WedgeIndex);
-			OutRawMesh.WedgeColors.Add(VertexColor);
-		}
+		FVertexID VertexID = OutRawMesh.CreateVertex();
+		VertexPositions[VertexID] = StaticMeshLOD.VertexBuffers.PositionVertexBuffer.VertexPosition(VertexIndex);
 	}
 
-	// Materials
+	//Create the vertex instances
+	for (int32 TriangleIndex = 0; TriangleIndex < NumFaces; ++TriangleIndex)
 	{
-		OutRawMesh.FaceMaterialIndices.Empty(NumFaces);
-		OutRawMesh.FaceMaterialIndices.SetNumZeroed(NumFaces);
-
-		for (const FStaticMeshSection& Section : StaticMeshLOD.Sections)
+		FPolygonGroupID CurrentPolygonGroupID = FPolygonGroupID::Invalid;
+		for (int32 SectionIndex = 0; SectionIndex < StaticMeshLOD.Sections.Num(); ++SectionIndex)
 		{
+			const FStaticMeshSection& Section = StaticMeshLOD.Sections[SectionIndex];
 			uint32 FirstTriangle = Section.FirstIndex / 3;
-			for (uint32 TriangleIndex = 0; TriangleIndex < Section.NumTriangles; ++TriangleIndex)
+			uint32 LastTriangle = FirstTriangle + Section.NumTriangles - 1;
+			if ((uint32)TriangleIndex >= FirstTriangle && (uint32)TriangleIndex <= LastTriangle)
 			{
-				OutRawMesh.FaceMaterialIndices[FirstTriangle + TriangleIndex] = Section.MaterialIndex;
+				CurrentPolygonGroupID = FPolygonGroupID(SectionIndex);
+				break;
 			}
 		}
-	}
+		check(CurrentPolygonGroupID != FPolygonGroupID::Invalid);
 
-	// Smoothing masks
-	{
-		OutRawMesh.FaceSmoothingMasks.Empty(NumFaces);
-		OutRawMesh.FaceSmoothingMasks.SetNumUninitialized(NumFaces);
+		FVertexID VertexIDs[3];
+		FVertexInstanceID VertexInstanceIDs[3];
+		for (int32 Corner = 0; Corner < 3; ++Corner)
+		{
+			int32 WedgeIndex = StaticMeshLOD.IndexBuffer.GetIndex(TriangleIndex * 3 + Corner);
+			FVertexID VertexID(WedgeIndex);
+			FVertexInstanceID VertexInstanceID = OutRawMesh.CreateVertexInstance(VertexID);
+			VertexIDs[Corner] = VertexID;
+			VertexInstanceIDs[Corner] = VertexInstanceID;
+
+			//NTBs
+			FVector TangentX = StaticMeshLOD.VertexBuffers.StaticMeshVertexBuffer.VertexTangentX(WedgeIndex);
+			FVector TangentY = StaticMeshLOD.VertexBuffers.StaticMeshVertexBuffer.VertexTangentY(WedgeIndex);
+			FVector TangentZ = StaticMeshLOD.VertexBuffers.StaticMeshVertexBuffer.VertexTangentZ(WedgeIndex);
+			VertexInstanceTangents[VertexInstanceID] = TangentX;
+			VertexInstanceBinormalSigns[VertexInstanceID] = GetBasisDeterminantSign(TangentX, TangentY, TangentZ);
+			VertexInstanceNormals[VertexInstanceID] = TangentZ;
+
+			// Vertex colors
+			if (StaticMeshLOD.VertexBuffers.ColorVertexBuffer.GetNumVertices() > 0)
+			{
+				VertexInstanceColors[VertexInstanceID] = FLinearColor(StaticMeshLOD.VertexBuffers.ColorVertexBuffer.VertexColor(WedgeIndex));
+			}
+			else
+			{
+				VertexInstanceColors[VertexInstanceID] = FLinearColor::White;
+			}
 
-		for (auto& SmoothingMask : OutRawMesh.FaceSmoothingMasks)
+			//Tex coord
+			for (int32 TexCoodIdx = 0; TexCoodIdx < NumTexCoords; ++TexCoodIdx)
+			{
+				VertexInstanceUVs.Set(VertexInstanceID, TexCoodIdx, StaticMeshLOD.VertexBuffers.StaticMeshVertexBuffer.GetVertexUV(WedgeIndex, TexCoodIdx));
+			}
+		}
+		//Create a polygon from this triangle
+		TArray<FMeshDescription::FContourPoint> Contours;
+		for (int32 Corner = 0; Corner < 3; ++Corner)
 		{
-			SmoothingMask = 1;
+			int32 ContourPointIndex = Contours.AddDefaulted();
+			FMeshDescription::FContourPoint& ContourPoint = Contours[ContourPointIndex];
+			//Find the matching edge ID
+			uint32 CornerIndices[2];
+			CornerIndices[0] = (Corner + 0) % 3;
+			CornerIndices[1] = (Corner + 1) % 3;
+
+			FVertexID EdgeVertexIDs[2];
+			EdgeVertexIDs[0] = VertexIDs[CornerIndices[0]];
+			EdgeVertexIDs[1] = VertexIDs[CornerIndices[1]];
+
+			FEdgeID MatchEdgeId = OutRawMesh.GetVertexPairEdge(EdgeVertexIDs[0], EdgeVertexIDs[1]);
+			if (MatchEdgeId == FEdgeID::Invalid)
+			{
+				MatchEdgeId = OutRawMesh.CreateEdge(EdgeVertexIDs[0], EdgeVertexIDs[1]);
+				EdgeHardnesses[MatchEdgeId] = false;
+				EdgeCreaseSharpnesses[MatchEdgeId] = 0.0f;
+			}
+			ContourPoint.EdgeID = MatchEdgeId;
+			ContourPoint.VertexInstanceID = VertexInstanceIDs[CornerIndices[0]];
 		}
+		// Insert a polygon into the mesh
+		const FPolygonID NewPolygonID = OutRawMesh.CreatePolygon(CurrentPolygonGroupID, Contours);
+		//Triangulate the polygon
+		FMeshPolygon& Polygon = OutRawMesh.GetPolygon(NewPolygonID);
+		OutRawMesh.ComputePolygonTriangulation(NewPolygonID, Polygon.Triangles);
 	}
 }
 
@@ -479,7 +584,31 @@ bool FMeshMergeHelpers::CheckWrappingUVs(const TArray<FVector2D>& UVs)
 	return bResult;
 }
 
-void FMeshMergeHelpers::CullTrianglesFromVolumesAndUnderLandscapes(const UWorld* World, const FBoxSphereBounds& Bounds, FRawMesh& InOutRawMesh)
+bool FMeshMergeHelpers::CheckWrappingUVs(const FMeshDescription& MeshDescription, int32 UVChannelIndex)
+{
+	TVertexInstanceAttributesConstRef<FVector2D> VertexInstanceUVs = MeshDescription.VertexInstanceAttributes().GetAttributesRef<FVector2D>(MeshAttribute::VertexInstance::TextureCoordinate);
+	bool bResult = false;
+
+	//Validate the channel, return false if there is an invalid channel index
+	if (UVChannelIndex < 0 || UVChannelIndex >= VertexInstanceUVs.GetNumIndices())
+	{
+		return bResult;
+	}
+
+	for (const FVertexInstanceID& VertexInstanceID : MeshDescription.VertexInstances().GetElementIDs())
+	{
+		const FVector2D& Coordinate = VertexInstanceUVs.Get(VertexInstanceID, UVChannelIndex);
+		if ((FMath::IsNegativeFloat(Coordinate.X) || FMath::IsNegativeFloat(Coordinate.Y)) || (Coordinate.X > (1.0f + KINDA_SMALL_NUMBER) || Coordinate.Y > (1.0f + KINDA_SMALL_NUMBER)))
+		{
+			bResult = true;
+			break;
+		}
+	}
+
+	return bResult;
+}
+
+void FMeshMergeHelpers::CullTrianglesFromVolumesAndUnderLandscapes(const UWorld* World, const FBoxSphereBounds& Bounds, FMeshDescription& InOutRawMesh)
 {
 	TArray<ALandscapeProxy*> Landscapes;
 	TArray<AMeshMergeCullingVolume*> CullVolumes;
@@ -518,14 +647,16 @@ void FMeshMergeHelpers::CullTrianglesFromVolumesAndUnderLandscapes(const UWorld*
 		}
 	}
 
-	TArray<bool> VertexVisible;
-	VertexVisible.AddZeroed(InOutRawMesh.VertexPositions.Num());
-	int32 Index = 0;
+	TVertexAttributesConstRef<FVector> VertexPositions = InOutRawMesh.VertexAttributes().GetAttributesRef<FVector>(MeshAttribute::Vertex::Position);
 
-	for (const FVector& Position : InOutRawMesh.VertexPositions)
+	TMap<FVertexID, bool> VertexVisible;
+	VertexVisible.Reserve(InOutRawMesh.Vertices().Num());
+	int32 Index = 0;
+	for(const FVertexID& VertexID : InOutRawMesh.Vertices().GetElementIDs())
 	{
+		const FVector& Position = VertexPositions[VertexID];
 		// Start with setting visibility to true on all vertices
-		VertexVisible[Index] = true;
+		VertexVisible.Add(VertexID, true);
 
 		// Check if this vertex is culled due to being underneath a landscape
 		if (Landscapes.Num() > 0)
@@ -551,7 +682,7 @@ void FMeshMergeHelpers::CullTrianglesFromVolumesAndUnderLandscapes(const UWorld*
 				const bool IsUnderneathLandscape = IsLandscapeHit(Start, End, World, Landscapes, OutHit);
 
 				// Vertex is visible when above landscape (with actual landscape underneath) or if there is no landscape beneath or above the vertex (falls outside of landscape bounds)
-				VertexVisible[Index] = (IsAboveLandscape && !IsUnderneathLandscape);// || (!IsAboveLandscape && !IsUnderneathLandscape);
+				VertexVisible[VertexID] = (IsAboveLandscape && !IsUnderneathLandscape);// || (!IsAboveLandscape && !IsUnderneathLandscape);
 			}
 		}
 
@@ -562,7 +693,7 @@ void FMeshMergeHelpers::CullTrianglesFromVolumesAndUnderLandscapes(const UWorld*
 			if (bVertexIsInsideVolume)
 			{
 				// Inside a culling volume so invisible
-				VertexVisible[Index] = false;
+				VertexVisible[VertexID] = false;
 			}
 		}
 
@@ -571,111 +702,92 @@ void FMeshMergeHelpers::CullTrianglesFromVolumesAndUnderLandscapes(const UWorld*
 
 
 	// We now know which vertices are below the landscape
-	TArray<bool> TriangleVisible;
-	int32 NumTriangles = InOutRawMesh.WedgeIndices.Num() / 3;
-	TriangleVisible.AddZeroed(NumTriangles);
-
-	bool bCreateNewMesh = false;
-
-	// Determine which triangles of the mesh are visible
-	for (int32 TriangleIndex = 0; TriangleIndex < NumTriangles; TriangleIndex++)
+	TArray<FPolygonID> PolygonToRemove;
+	for(const FPolygonID& PolygonID : InOutRawMesh.Polygons().GetElementIDs())
 	{
 		bool AboveLandscape = false;
-
-		for (int32 WedgeIndex = 0; WedgeIndex < 3; ++WedgeIndex)
+		for (FMeshTriangle& Triangle : InOutRawMesh.GetPolygonTriangles(PolygonID))
 		{
-			AboveLandscape |= VertexVisible[InOutRawMesh.WedgeIndices[(TriangleIndex * 3) + WedgeIndex]];
+			for (int32 Corner = 0; Corner < 3; ++Corner)
+			{
+				AboveLandscape |= VertexVisible[InOutRawMesh.GetVertexInstanceVertex(Triangle.GetVertexInstanceID(Corner))];
+			}
+		}
+		if (!AboveLandscape)
+		{
+			PolygonToRemove.Add(PolygonID);
 		}
-		TriangleVisible[TriangleIndex] = AboveLandscape;
-		bCreateNewMesh |= !AboveLandscape;
-
 	}
 
-	// Check whether or not we have to create a new mesh
-	if (bCreateNewMesh)
+	// Delete the polygons that are not visible
 	{
-		FRawMesh NewRawMesh;
-		TMap<int32, int32> VertexRemapping;
-
-		// Fill new mesh with data only from visible triangles
-		for (int32 TriangleIndex = 0; TriangleIndex < NumTriangles; ++TriangleIndex)
+		TArray<FEdgeID> OrphanedEdges;
+		TArray<FVertexInstanceID> OrphanedVertexInstances;
+		TArray<FPolygonGroupID> OrphanedPolygonGroups;
+		TArray<FVertexID> OrphanedVertices;
+		for (FPolygonID PolygonID : PolygonToRemove)
 		{
-			if (!TriangleVisible[TriangleIndex])
-				continue;
-
-			for (int32 WedgeIndex = 0; WedgeIndex < 3; ++WedgeIndex)
-			{
-				int32 OldIndex = InOutRawMesh.WedgeIndices[(TriangleIndex * 3) + WedgeIndex];
-
-				int32 NewIndex;
-
-				int32* RemappedIndex = VertexRemapping.Find(Index);
-				if (RemappedIndex)
-				{
-					NewIndex = *RemappedIndex;
-				}
-				else
-				{
-					NewIndex = NewRawMesh.VertexPositions.Add(InOutRawMesh.VertexPositions[OldIndex]);
-					VertexRemapping.Add(OldIndex, NewIndex);
-				}
-
-				NewRawMesh.WedgeIndices.Add(NewIndex);
-				if (InOutRawMesh.WedgeColors.Num()) NewRawMesh.WedgeColors.Add(InOutRawMesh.WedgeColors[(TriangleIndex * 3) + WedgeIndex]);
-				if (InOutRawMesh.WedgeTangentX.Num()) NewRawMesh.WedgeTangentX.Add(InOutRawMesh.WedgeTangentX[(TriangleIndex * 3) + WedgeIndex]);
-				if (InOutRawMesh.WedgeTangentY.Num()) NewRawMesh.WedgeTangentY.Add(InOutRawMesh.WedgeTangentY[(TriangleIndex * 3) + WedgeIndex]);
-				if (InOutRawMesh.WedgeTangentZ.Num()) NewRawMesh.WedgeTangentZ.Add(InOutRawMesh.WedgeTangentZ[(TriangleIndex * 3) + WedgeIndex]);
-
-				for (int32 UVIndex = 0; UVIndex < MAX_MESH_TEXTURE_COORDS; ++UVIndex)
-				{
-					if (InOutRawMesh.WedgeTexCoords[UVIndex].Num())
-					{
-						NewRawMesh.WedgeTexCoords[UVIndex].Add(InOutRawMesh.WedgeTexCoords[UVIndex][(TriangleIndex * 3) + WedgeIndex]);
-					}
-				}
-			}
-
-			NewRawMesh.FaceMaterialIndices.Add(InOutRawMesh.FaceMaterialIndices[TriangleIndex]);
-			NewRawMesh.FaceSmoothingMasks.Add(InOutRawMesh.FaceSmoothingMasks[TriangleIndex]);
+			InOutRawMesh.DeletePolygon(PolygonID, &OrphanedEdges, &OrphanedVertexInstances, &OrphanedPolygonGroups);
 		}
-
-		InOutRawMesh = NewRawMesh;
+		//Do not remove the polygongroup since its indexed with the mesh material array
+		/*for (FPolygonGroupID PolygonGroupID : OrphanedPolygonGroups)
+		{
+			InOutRawMesh.DeletePolygonGroup(PolygonGroupID);
+		}*/
+		for (FVertexInstanceID VertexInstanceID : OrphanedVertexInstances)
+		{
+			InOutRawMesh.DeleteVertexInstance(VertexInstanceID, &OrphanedVertices);
+		}
+		for (FEdgeID EdgeID : OrphanedEdges)
+		{
+			InOutRawMesh.DeleteEdge(EdgeID, &OrphanedVertices);
+		}
+		for (FVertexID VertexID : OrphanedVertices)
+		{
+			InOutRawMesh.DeleteVertex(VertexID);
+		}
+		//Compact and Remap IDs so we have clean ID from 0 to n since we just erase some polygons
+		//The render build need to have compact ID
+		FElementIDRemappings OutRemappings;
+		InOutRawMesh.Compact(OutRemappings);
 	}
 }
 
-void FMeshMergeHelpers::PropagateSplineDeformationToRawMesh(const USplineMeshComponent* InSplineMeshComponent, struct FRawMesh &OutRawMesh)
+void FMeshMergeHelpers::PropagateSplineDeformationToRawMesh(const USplineMeshComponent* InSplineMeshComponent, FMeshDescription &OutRawMesh)
 {
+	TVertexAttributesRef<FVector> VertexPositions = OutRawMesh.VertexAttributes().GetAttributesRef<FVector>(MeshAttribute::Vertex::Position);
+	TVertexInstanceAttributesRef<FVector> VertexInstanceNormals = OutRawMesh.VertexInstanceAttributes().GetAttributesRef<FVector>(MeshAttribute::VertexInstance::Normal);
+	TVertexInstanceAttributesRef<FVector> VertexInstanceTangents = OutRawMesh.VertexInstanceAttributes().GetAttributesRef<FVector>(MeshAttribute::VertexInstance::Tangent);
+	TVertexInstanceAttributesRef<float> VertexInstanceBinormalSigns = OutRawMesh.VertexInstanceAttributes().GetAttributesRef<float>(MeshAttribute::VertexInstance::BinormalSign);
+
 	// Apply spline deformation for each vertex's tangents
-	for (int32 iVert = 0; iVert < OutRawMesh.WedgeIndices.Num(); ++iVert)
+	int32 WedgeIndex = 0;
+	for (const FPolygonID& PolygonID : OutRawMesh.Polygons().GetElementIDs())
 	{
-		uint32 Index = OutRawMesh.WedgeIndices[iVert];
-		float& AxisValue = USplineMeshComponent::GetAxisValue(OutRawMesh.VertexPositions[Index], InSplineMeshComponent->ForwardAxis);
-		FTransform SliceTransform = InSplineMeshComponent->CalcSliceTransform(AxisValue);
-
-		// Transform tangents first
-		if (OutRawMesh.WedgeTangentX.Num())
+		for (FMeshTriangle& Triangle : OutRawMesh.GetPolygonTriangles(PolygonID))
 		{
-			OutRawMesh.WedgeTangentX[iVert] = SliceTransform.TransformVector(OutRawMesh.WedgeTangentX[iVert]);
-		}
-
-		if (OutRawMesh.WedgeTangentY.Num())
-		{
-			OutRawMesh.WedgeTangentY[iVert] = SliceTransform.TransformVector(OutRawMesh.WedgeTangentY[iVert]);
-		}
-
-		if (OutRawMesh.WedgeTangentZ.Num())
-		{
-			OutRawMesh.WedgeTangentZ[iVert] = SliceTransform.TransformVector(OutRawMesh.WedgeTangentZ[iVert]);
+			for (int32 Corner = 0; Corner < 3; ++Corner, ++WedgeIndex)
+			{
+				const FVertexInstanceID VertexInstanceID = Triangle.GetVertexInstanceID(Corner);
+				const FVertexID VertexID = OutRawMesh.GetVertexInstanceVertex(VertexInstanceID);
+				const float& AxisValue = USplineMeshComponent::GetAxisValue(VertexPositions[VertexID], InSplineMeshComponent->ForwardAxis);
+				FTransform SliceTransform = InSplineMeshComponent->CalcSliceTransform(AxisValue);
+				FVector TangentY = FVector::CrossProduct(VertexInstanceNormals[VertexInstanceID], VertexInstanceTangents[VertexInstanceID]).GetSafeNormal() * VertexInstanceBinormalSigns[VertexInstanceID];
+				VertexInstanceTangents[VertexInstanceID] = SliceTransform.TransformVector(VertexInstanceTangents[VertexInstanceID]);
+				TangentY = SliceTransform.TransformVector(TangentY);
+				VertexInstanceNormals[VertexInstanceID] = SliceTransform.TransformVector(VertexInstanceNormals[VertexInstanceID]);
+				VertexInstanceBinormalSigns[VertexInstanceID] = GetBasisDeterminantSign(VertexInstanceTangents[VertexInstanceID], TangentY, VertexInstanceNormals[VertexInstanceID]);
+			}
 		}
 	}
 
 	// Apply spline deformation for each vertex position
-	for (int32 iVert = 0; iVert < OutRawMesh.VertexPositions.Num(); ++iVert)
+	for (const FVertexID& VertexID : OutRawMesh.Vertices().GetElementIDs())
 	{
-		float& AxisValue = USplineMeshComponent::GetAxisValue(OutRawMesh.VertexPositions[iVert], InSplineMeshComponent->ForwardAxis);
+		float& AxisValue = USplineMeshComponent::GetAxisValue(VertexPositions[VertexID], InSplineMeshComponent->ForwardAxis);
 		FTransform SliceTransform = InSplineMeshComponent->CalcSliceTransform(AxisValue);
 		AxisValue = 0.0f;
-		OutRawMesh.VertexPositions[iVert] = SliceTransform.TransformPosition(OutRawMesh.VertexPositions[iVert]);
+		VertexPositions[VertexID] = SliceTransform.TransformPosition(VertexPositions[VertexID]);
 	}
 }
 
@@ -708,11 +820,21 @@ void FMeshMergeHelpers::PropagateSplineDeformationToPhysicsGeometry(USplineMeshC
 	}
 }
 
-void FMeshMergeHelpers::TransformRawMeshVertexData(const FTransform& InTransform, FRawMesh &OutRawMesh)
+void FMeshMergeHelpers::TransformRawMeshVertexData(const FTransform& InTransform, FMeshDescription &OutRawMesh)
 {
-	for (FVector& Vertex : OutRawMesh.VertexPositions)
-	{
-		Vertex = InTransform.TransformPosition(Vertex);
+	TVertexAttributesRef<FVector> VertexPositions = OutRawMesh.VertexAttributes().GetAttributesRef<FVector>(MeshAttribute::Vertex::Position);
+	TEdgeAttributesRef<bool> EdgeHardnesses = OutRawMesh.EdgeAttributes().GetAttributesRef<bool>(MeshAttribute::Edge::IsHard);
+	TEdgeAttributesRef<float> EdgeCreaseSharpnesses = OutRawMesh.EdgeAttributes().GetAttributesRef<float>(MeshAttribute::Edge::CreaseSharpness);
+	TPolygonGroupAttributesRef<FName> PolygonGroupImportedMaterialSlotNames = OutRawMesh.PolygonGroupAttributes().GetAttributesRef<FName>(MeshAttribute::PolygonGroup::ImportedMaterialSlotName);
+	TVertexInstanceAttributesRef<FVector> VertexInstanceNormals = OutRawMesh.VertexInstanceAttributes().GetAttributesRef<FVector>(MeshAttribute::VertexInstance::Normal);
+	TVertexInstanceAttributesRef<FVector> VertexInstanceTangents = OutRawMesh.VertexInstanceAttributes().GetAttributesRef<FVector>(MeshAttribute::VertexInstance::Tangent);
+	TVertexInstanceAttributesRef<float> VertexInstanceBinormalSigns = OutRawMesh.VertexInstanceAttributes().GetAttributesRef<float>(MeshAttribute::VertexInstance::BinormalSign);
+	TVertexInstanceAttributesRef<FVector4> VertexInstanceColors = OutRawMesh.VertexInstanceAttributes().GetAttributesRef<FVector4>(MeshAttribute::VertexInstance::Color);
+	TVertexInstanceAttributesRef<FVector2D> VertexInstanceUVs = OutRawMesh.VertexInstanceAttributes().GetAttributesRef<FVector2D>(MeshAttribute::VertexInstance::TextureCoordinate);
+
+	for(const FVertexID& VertexID : OutRawMesh.Vertices().GetElementIDs())
+	{
+		VertexPositions[VertexID] = InTransform.TransformPosition(VertexPositions[VertexID]);
 	}
 	
 	auto TransformNormal = [&](FVector& Normal)
@@ -729,39 +851,24 @@ void FMeshMergeHelpers::TransformRawMeshVertexData(const FTransform& InTransform
 		}
 	};	
 
-	for (FVector& TangentX : OutRawMesh.WedgeTangentX)
-	{
-		TransformNormal(TangentX);
-	}
-
-	for (FVector& TangentY : OutRawMesh.WedgeTangentY)
+	for(const FVertexInstanceID& VertexInstanceID : OutRawMesh.VertexInstances().GetElementIDs())
 	{
+		FVector TangentY = FVector::CrossProduct(VertexInstanceNormals[VertexInstanceID], VertexInstanceTangents[VertexInstanceID]).GetSafeNormal() * VertexInstanceBinormalSigns[VertexInstanceID];
+		TransformNormal(VertexInstanceTangents[VertexInstanceID]);
 		TransformNormal(TangentY);
-	}
-
-	for (FVector& TangentZ : OutRawMesh.WedgeTangentZ)
-	{
-		TransformNormal(TangentZ);
+		TransformNormal(VertexInstanceNormals[VertexInstanceID]);
+		VertexInstanceBinormalSigns[VertexInstanceID] = GetBasisDeterminantSign(VertexInstanceTangents[VertexInstanceID], TangentY, VertexInstanceNormals[VertexInstanceID]);
 	}
 
 	const bool bIsMirrored = InTransform.GetDeterminant() < 0.f;
 	if (bIsMirrored)
 	{
-		Algo::Reverse(OutRawMesh.WedgeIndices);
-		Algo::Reverse(OutRawMesh.WedgeTangentX);
-		Algo::Reverse(OutRawMesh.WedgeTangentY);
-		Algo::Reverse(OutRawMesh.WedgeTangentZ);
-		for (uint32 UVIndex = 0; UVIndex < MAX_MESH_TEXTURE_COORDS; ++UVIndex)
-		{
-			Algo::Reverse(OutRawMesh.WedgeTexCoords[UVIndex]);
-		}
-		Algo::Reverse(OutRawMesh.FaceMaterialIndices);
-		Algo::Reverse(OutRawMesh.FaceSmoothingMasks);
-		Algo::Reverse(OutRawMesh.WedgeColors);
+		//Reverse the vertexinstance
+		OutRawMesh.ReverseAllPolygonFacing();
 	}
 }
 
-void FMeshMergeHelpers::RetrieveCullingLandscapeAndVolumes(UWorld* InWorld, const FBoxSphereBounds& EstimatedMeshProxyBounds, const TEnumAsByte<ELandscapeCullingPrecision::Type> PrecisionType, TArray<FRawMesh*>& CullingRawMeshes)
+void FMeshMergeHelpers::RetrieveCullingLandscapeAndVolumes(UWorld* InWorld, const FBoxSphereBounds& EstimatedMeshProxyBounds, const TEnumAsByte<ELandscapeCullingPrecision::Type> PrecisionType, TArray<FMeshDescription*>& CullingRawMeshes)
 {
 	// Extract landscape proxies and cull volumes from the world
 	TArray<ALandscapeProxy*> LandscapeActors;
@@ -810,16 +917,7 @@ void FMeshMergeHelpers::RetrieveCullingLandscapeAndVolumes(UWorld* InWorld, cons
 		Landscape->ExportToRawMesh(LandscapeExportLOD, MeshDescription, LandscapeBounds);
 		if (MeshDescription.Vertices().Num())
 		{
-			TMap<FName, int32> MaterialMap;
-			TPolygonGroupAttributesConstRef<FName> PolygonGroupNames = MeshDescription.PolygonGroupAttributes().GetAttributesRef<FName>(MeshAttribute::PolygonGroup::ImportedMaterialSlotName);
-			int32 MaterialIndex = 0;
-			for (const FPolygonGroupID PolygonGroupID : MeshDescription.PolygonGroups().GetElementIDs())
-			{
-				MaterialMap.Add(PolygonGroupNames[PolygonGroupID], MaterialIndex++);
-			}
-			FRawMesh* LandscapeRawMesh = new FRawMesh();
-			FMeshDescriptionOperations::ConvertToRawMesh(MeshDescription, *LandscapeRawMesh, MaterialMap);
-			CullingRawMeshes.Add(LandscapeRawMesh);
+			CullingRawMeshes.Add(&MeshDescription);
 		}
 	}
 
@@ -827,16 +925,18 @@ void FMeshMergeHelpers::RetrieveCullingLandscapeAndVolumes(UWorld* InWorld, cons
 	for (AMeshMergeCullingVolume* Volume : CullVolumes)
 	{
 		// Export the landscape to raw mesh format
-		FRawMesh* VolumeMesh = new FRawMesh();
+		FMeshDescription* VolumeMesh = new FMeshDescription();
+		UStaticMesh::RegisterMeshAttributes(*VolumeMesh);
 
 		TArray<FStaticMaterial>	VolumeMaterials;
 		GetBrushMesh(Volume, Volume->Brush, *VolumeMesh, VolumeMaterials);
 
 		// Offset vertices to correct world position;
 		FVector VolumeLocation = Volume->GetActorLocation();
-		for (FVector& Position : VolumeMesh->VertexPositions)
+		TVertexAttributesRef<FVector> VertexPositions = VolumeMesh->VertexAttributes().GetAttributesRef<FVector>(MeshAttribute::Vertex::Position);
+		for(const FVertexID& VertexID : VolumeMesh->Vertices().GetElementIDs())
 		{
-			Position += VolumeLocation;
+			VertexPositions[VertexID] += VolumeLocation;
 		}
 
 		CullingRawMeshes.Add(VolumeMesh);
@@ -965,35 +1065,29 @@ FVector2D FMeshMergeHelpers::GetValidUV(const FVector2D& UV)
 	return NewUV;
 }
 
-void FMeshMergeHelpers::CalculateTextureCoordinateBoundsForRawMesh(const FRawMesh& InRawMesh, TArray<FBox2D>& OutBounds)
+void FMeshMergeHelpers::CalculateTextureCoordinateBoundsForRawMesh(const FMeshDescription& InRawMesh, TArray<FBox2D>& OutBounds)
 {
-	const int32 NumWedges = InRawMesh.WedgeIndices.Num();
-	const int32 NumTris = NumWedges / 3;
-
+	TVertexInstanceAttributesConstRef<FVector2D> VertexInstanceUVs = InRawMesh.VertexInstanceAttributes().GetAttributesRef<FVector2D>(MeshAttribute::VertexInstance::TextureCoordinate);
 	OutBounds.Empty();
-	for (int32 TriIndex = 0; TriIndex < NumTris; TriIndex++)
+	for (const FPolygonID& PolygonID : InRawMesh.Polygons().GetElementIDs())
 	{
-		int MaterialIndex = InRawMesh.FaceMaterialIndices[TriIndex];
+		int32 MaterialIndex = InRawMesh.GetPolygonPolygonGroup(PolygonID).GetValue();
 		if (OutBounds.Num() <= MaterialIndex)
 			OutBounds.SetNumZeroed(MaterialIndex + 1);
 		{
-			const int32 CachedWedgeIndex = TriIndex * 3;
-			for (int32 UVIndex = 0; UVIndex < MAX_MESH_TEXTURE_COORDS; ++UVIndex)
+			TArray<FVertexInstanceID> PolygonVertexInstances = InRawMesh.GetPolygonPerimeterVertexInstances(PolygonID);
+			for (const FVertexInstanceID& VertexInstanceID : PolygonVertexInstances)
 			{
-				int32 WedgeIndex = CachedWedgeIndex;
-				if (InRawMesh.WedgeTexCoords[UVIndex].Num())
+				for (int32 UVIndex = 0; UVIndex < VertexInstanceUVs.GetNumIndices(); ++UVIndex)
 				{
-					for (int32 CornerIndex = 0; CornerIndex < 3; CornerIndex++, WedgeIndex++)
-					{
-						OutBounds[MaterialIndex] += InRawMesh.WedgeTexCoords[UVIndex][WedgeIndex];
-					}
+					OutBounds[MaterialIndex] += VertexInstanceUVs.Get(VertexInstanceID, UVIndex);
 				}
 			}
 		}
 	}
 }
 
-bool FMeshMergeHelpers::PropagatePaintedColorsToRawMesh(const UStaticMeshComponent* StaticMeshComponent, int32 LODIndex, FRawMesh& RawMesh)
+bool FMeshMergeHelpers::PropagatePaintedColorsToRawMesh(const UStaticMeshComponent* StaticMeshComponent, int32 LODIndex, FMeshDescription& RawMesh)
 {
 	UStaticMesh* StaticMesh = StaticMeshComponent->GetStaticMesh();
 
@@ -1006,16 +1100,28 @@ bool FMeshMergeHelpers::PropagatePaintedColorsToRawMesh(const UStaticMeshCompone
 
 		if (ColorVertexBuffer.GetNumVertices() == RenderModel.GetNumVertices())
 		{	
-			const int32 NumWedges = RawMesh.WedgeIndices.Num();
+			const int32 NumWedges = RawMesh.VertexInstances().Num();
 			const int32 NumRenderWedges = RenderModel.IndexBuffer.GetNumIndices();
 			const bool bUseRenderWedges = NumWedges == NumRenderWedges;
-					
+
+			TVertexInstanceAttributesRef<FVector4> VertexInstanceColors = RawMesh.VertexInstanceAttributes().GetAttributesRef<FVector4>(MeshAttribute::VertexInstance::Color);
+
 			if (bUseRenderWedges)
 			{
-				const int32 NumExistingColors = RawMesh.WedgeColors.Num();
-				if (NumExistingColors < NumRenderWedges)
+				//Create a map index
+				TMap<int32, FVertexInstanceID> IndexToVertexInstanceID;
+				IndexToVertexInstanceID.Reserve(NumWedges);
+				int32 CurrentWedgeIndex = 0;
+				for (const FPolygonID& PolygonID : RawMesh.Polygons().GetElementIDs())
 				{
-					RawMesh.WedgeColors.AddUninitialized(NumRenderWedges - NumExistingColors);
+					const TArray<FMeshTriangle>& Triangles = RawMesh.GetPolygonTriangles(PolygonID);
+					for (const FMeshTriangle& Triangle : Triangles)
+					{
+						for (int32 Corner = 0; Corner < 3; ++Corner, ++CurrentWedgeIndex)
+						{
+							IndexToVertexInstanceID.Add(CurrentWedgeIndex, Triangle.GetVertexInstanceID(Corner));
+						}
+					}
 				}
 
 				const FIndexArrayView ArrayView = RenderModel.IndexBuffer.GetArrayView();
@@ -1027,32 +1133,41 @@ bool FMeshMergeHelpers::PropagatePaintedColorsToRawMesh(const UStaticMeshCompone
 					{
 						WedgeColor = ColorVertexBuffer.VertexColor(Index);
 					}
-
-					RawMesh.WedgeColors[WedgeIndex] = WedgeColor;
+					VertexInstanceColors[IndexToVertexInstanceID[WedgeIndex]] = FLinearColor(WedgeColor);
 				}
 
 				return true;				
 			}
 			// No wedge map (this can happen when we poly reduce the LOD for example)
-			// Use index buffer directly
+			// Use index buffer directly. Not sure this will happen with FMeshDescription
 			else
 			{
-				RawMesh.WedgeColors.SetNumUninitialized(NumWedges);
-
-				if (RawMesh.VertexPositions.Num() == ColorVertexBuffer.GetNumVertices())
+				if (RawMesh.Vertices().Num() == ColorVertexBuffer.GetNumVertices())
 				{
-					for (int32 WedgeIndex = 0; WedgeIndex < NumWedges; ++WedgeIndex)
+					//Create a map index
+					TMap<FVertexID, int32> VertexIDToVertexIndex;
+					VertexIDToVertexIndex.Reserve(RawMesh.Vertices().Num());
+					int32 CurrentVertexIndex = 0;
+					for (const FVertexID& VertexID : RawMesh.Vertices().GetElementIDs())
+					{
+						VertexIDToVertexIndex.Add(VertexID, CurrentVertexIndex++);
+					}
+
+					for (const FVertexID& VertexID : RawMesh.Vertices().GetElementIDs())
 					{
 						FColor WedgeColor = FColor::White;
-						uint32 VertIndex = RawMesh.WedgeIndices[WedgeIndex];
+						uint32 VertIndex = VertexIDToVertexIndex[VertexID];
 
 						if (VertIndex < ColorVertexBuffer.GetNumVertices())
 						{
 							WedgeColor = ColorVertexBuffer.VertexColor(VertIndex);
 						}
-						RawMesh.WedgeColors[WedgeIndex] = WedgeColor;
+						const TArray<FVertexInstanceID>& VertexInstances = RawMesh.GetVertexVertexInstances(VertexID);
+						for (const FVertexInstanceID& VertexInstanceID : VertexInstances)
+						{
+							VertexInstanceColors[VertexInstanceID] = FLinearColor(WedgeColor);
+						}
 					}
-
 					return true;
 				}
 			}
@@ -1090,68 +1205,137 @@ bool FMeshMergeHelpers::IsLandscapeHit(const FVector& RayOrigin, const FVector&
 	return bHitLandscape;
 }
 
-void FMeshMergeHelpers::AppendRawMesh(FRawMesh& InTarget, const FRawMesh& InSource)
+void FMeshMergeHelpers::AppendRawMesh(FMeshDescription& InTarget, const FMeshDescription& InSource)
 {
-	const int32 MaxSmoothingMask = [InTarget]() -> int32
-	{
-		int32 Max = 0;
-		for (const int32 Value : InTarget.FaceSmoothingMasks)
+	TVertexAttributesConstRef<FVector> SourceVertexPositions = InSource.VertexAttributes().GetAttributesRef<FVector>(MeshAttribute::Vertex::Position);
+	TEdgeAttributesConstRef<bool> SourceEdgeHardnesses = InSource.EdgeAttributes().GetAttributesRef<bool>(MeshAttribute::Edge::IsHard);
+	TEdgeAttributesConstRef<float> SourceEdgeCreaseSharpnesses = InSource.EdgeAttributes().GetAttributesRef<float>(MeshAttribute::Edge::CreaseSharpness);
+	TPolygonGroupAttributesConstRef<FName> SourcePolygonGroupImportedMaterialSlotNames = InSource.PolygonGroupAttributes().GetAttributesRef<FName>(MeshAttribute::PolygonGroup::ImportedMaterialSlotName);
+	TVertexInstanceAttributesConstRef<FVector> SourceVertexInstanceNormals = InSource.VertexInstanceAttributes().GetAttributesRef<FVector>(MeshAttribute::VertexInstance::Normal);
+	TVertexInstanceAttributesConstRef<FVector> SourceVertexInstanceTangents = InSource.VertexInstanceAttributes().GetAttributesRef<FVector>(MeshAttribute::VertexInstance::Tangent);
+	TVertexInstanceAttributesConstRef<float> SourceVertexInstanceBinormalSigns = InSource.VertexInstanceAttributes().GetAttributesRef<float>(MeshAttribute::VertexInstance::BinormalSign);
+	TVertexInstanceAttributesConstRef<FVector4> SourceVertexInstanceColors = InSource.VertexInstanceAttributes().GetAttributesRef<FVector4>(MeshAttribute::VertexInstance::Color);
+	TVertexInstanceAttributesConstRef<FVector2D> SourceVertexInstanceUVs = InSource.VertexInstanceAttributes().GetAttributesRef<FVector2D>(MeshAttribute::VertexInstance::TextureCoordinate);
+
+	TVertexAttributesRef<FVector> TargetVertexPositions = InTarget.VertexAttributes().GetAttributesRef<FVector>(MeshAttribute::Vertex::Position);
+	TEdgeAttributesRef<bool> TargetEdgeHardnesses = InTarget.EdgeAttributes().GetAttributesRef<bool>(MeshAttribute::Edge::IsHard);
+	TEdgeAttributesRef<float> TargetEdgeCreaseSharpnesses = InTarget.EdgeAttributes().GetAttributesRef<float>(MeshAttribute::Edge::CreaseSharpness);
+	TPolygonGroupAttributesRef<FName> TargetPolygonGroupImportedMaterialSlotNames = InTarget.PolygonGroupAttributes().GetAttributesRef<FName>(MeshAttribute::PolygonGroup::ImportedMaterialSlotName);
+	TVertexInstanceAttributesRef<FVector> TargetVertexInstanceNormals = InTarget.VertexInstanceAttributes().GetAttributesRef<FVector>(MeshAttribute::VertexInstance::Normal);
+	TVertexInstanceAttributesRef<FVector> TargetVertexInstanceTangents = InTarget.VertexInstanceAttributes().GetAttributesRef<FVector>(MeshAttribute::VertexInstance::Tangent);
+	TVertexInstanceAttributesRef<float> TargetVertexInstanceBinormalSigns = InTarget.VertexInstanceAttributes().GetAttributesRef<float>(MeshAttribute::VertexInstance::BinormalSign);
+	TVertexInstanceAttributesRef<FVector4> TargetVertexInstanceColors = InTarget.VertexInstanceAttributes().GetAttributesRef<FVector4>(MeshAttribute::VertexInstance::Color);
+	TVertexInstanceAttributesRef<FVector2D> TargetVertexInstanceUVs = InTarget.VertexInstanceAttributes().GetAttributesRef<FVector2D>(MeshAttribute::VertexInstance::TextureCoordinate);
+
+	InTarget.ReserveNewVertices(InSource.Vertices().Num());
+	InTarget.ReserveNewVertexInstances(InSource.VertexInstances().Num());
+	InTarget.ReserveNewEdges(InSource.Edges().Num());
+	InTarget.ReserveNewPolygons(InSource.Vertices().Num());
+
+	//Append PolygonGroup
+	for (const FPolygonGroupID& SourcePolygonGroupID : InSource.PolygonGroups().GetElementIDs())
+	{
+		if (!InTarget.IsPolygonGroupValid(SourcePolygonGroupID))
 		{
-			Max = FMath::Max(Max, Value);
+			InTarget.CreatePolygonGroupWithID(SourcePolygonGroupID);
+			TargetPolygonGroupImportedMaterialSlotNames[SourcePolygonGroupID] = SourcePolygonGroupImportedMaterialSlotNames[SourcePolygonGroupID];
 		}
-		return Max;
-	}();
-
-	InTarget.FaceMaterialIndices.Append(InSource.FaceMaterialIndices);
-	
-	const int32 FaceOffset = InTarget.FaceSmoothingMasks.Num();
-	InTarget.FaceSmoothingMasks.Append(InSource.FaceSmoothingMasks);
+	}
 
-	for (int32 Index = FaceOffset; Index < InTarget.FaceSmoothingMasks.Num(); ++Index)
+	//Append the Vertexs
+	TMap<FVertexID, FVertexID> SourceToTargetVertexID;
+	SourceToTargetVertexID.Reserve(InSource.Vertices().Num());
+	for (const FVertexID& SourceVertexID : InSource.Vertices().GetElementIDs())
 	{
-		InTarget.FaceSmoothingMasks[Index] += FaceOffset;
+		const FVertexID TargetVertexID = InTarget.CreateVertex();
+		SourceToTargetVertexID.Add(SourceVertexID, TargetVertexID);
 	}
 
-	const int32 VertexOffset = InTarget.VertexPositions.Num();
-	InTarget.VertexPositions.Append(InSource.VertexPositions);
-	const int32 IndexOffset = InTarget.WedgeIndices.Num();
-	InTarget.WedgeIndices.Append(InSource.WedgeIndices);
-	for (int32 Index = IndexOffset; Index < InTarget.WedgeIndices.Num(); ++Index)
+	//Append VertexInstances
+	if (SourceVertexInstanceUVs.GetNumIndices() > TargetVertexInstanceUVs.GetNumIndices())
 	{
-		InTarget.WedgeIndices[Index] += VertexOffset;
+		TargetVertexInstanceUVs.SetNumIndices(SourceVertexInstanceUVs.GetNumIndices());
+	}
+	TMap<FVertexInstanceID, FVertexInstanceID> SourceToTargetVertexInstanceID;
+	SourceToTargetVertexInstanceID.Reserve(InSource.VertexInstances().Num());
+	for (const FVertexInstanceID& SourceVertexInstanceID : InSource.VertexInstances().GetElementIDs())
+	{
+		const FVertexID SourceVertexID = InSource.GetVertexInstanceVertex(SourceVertexInstanceID);
+		const FVertexInstanceID TargetVertexInstanceID = InTarget.CreateVertexInstance(SourceToTargetVertexID[SourceVertexID]);
+		TargetVertexInstanceTangents[TargetVertexInstanceID] = SourceVertexInstanceTangents[SourceVertexInstanceID];
+		TargetVertexInstanceBinormalSigns[TargetVertexInstanceID] = SourceVertexInstanceBinormalSigns[SourceVertexInstanceID];
+		TargetVertexInstanceNormals[TargetVertexInstanceID] = SourceVertexInstanceNormals[SourceVertexInstanceID];
+		TargetVertexInstanceColors[TargetVertexInstanceID] = SourceVertexInstanceColors[SourceVertexInstanceID];
+		for (int32 UVIndex = 0; UVIndex < TargetVertexInstanceUVs.GetNumIndices(); ++UVIndex)
+		{
+			FVector2D SourceUV = SourceVertexInstanceUVs.GetNumIndices() > UVIndex ? SourceVertexInstanceUVs.Get(SourceVertexInstanceID, UVIndex) : FVector2D(0.0f, 0.0f);
+			TargetVertexInstanceUVs.Set(TargetVertexInstanceID, UVIndex, SourceUV);
+		}
+		SourceToTargetVertexInstanceID.Add(SourceVertexInstanceID, TargetVertexInstanceID);
 	}
 
-	InTarget.WedgeTangentX.Append(InSource.WedgeTangentX);
-	InTarget.WedgeTangentY.Append(InSource.WedgeTangentY);
-	InTarget.WedgeTangentZ.Append(InSource.WedgeTangentZ);
-	
-	// Check whether or not we have pad the wedge colors 
-	const int32 NumTotalColors = InTarget.WedgeColors.Num() + InSource.WedgeColors.Num();
-	if (NumTotalColors < InTarget.WedgeTangentZ.Num() && (InSource.WedgeColors.Num() || InTarget.WedgeColors.Num()) )
+	//Append Edges
+	TMap<FEdgeID, FEdgeID> SourceToTargetEdgeID;
+	SourceToTargetEdgeID.Reserve(InSource.Edges().Num());
+	for (const FEdgeID& SourceEdgeID : InSource.Edges().GetElementIDs())
 	{
-		InTarget.WedgeColors.AddZeroed(InTarget.WedgeTangentZ.Num() - NumTotalColors);
+		const FMeshEdge& SourceEdge = InSource.GetEdge(SourceEdgeID);
+		const FEdgeID TargetEdgeID = InTarget.CreateEdge(SourceToTargetVertexID[SourceEdge.VertexIDs[0]], SourceToTargetVertexID[SourceEdge.VertexIDs[1]]);
+		TargetEdgeHardnesses[TargetEdgeID] = SourceEdgeHardnesses[SourceEdgeID];
+		TargetEdgeCreaseSharpnesses[TargetEdgeID] = SourceEdgeCreaseSharpnesses[SourceEdgeID];
+		SourceToTargetEdgeID.Add(SourceEdgeID, TargetEdgeID);
 	}
-	
-	InTarget.WedgeColors.Append(InSource.WedgeColors);
 
-	for (int32 i = 0; i < MAX_MESH_TEXTURE_COORDS; ++i)
+	auto CreateContour = [&InSource, &SourceToTargetVertexInstanceID, &SourceToTargetEdgeID](const TArray<FVertexInstanceID>& SourceVertexInstanceIDs, TArray<FMeshDescription::FContourPoint>& Contours)
 	{
-		const int32 NumTotalUVs = InTarget.WedgeTexCoords[i].Num() + InSource.WedgeTexCoords[i].Num();
+		const int32 ContourCount = SourceVertexInstanceIDs.Num();
+		for (int32 ContourIndex = 0; ContourIndex < ContourCount; ++ContourIndex)
+		{
+			FVertexInstanceID SourceVertexInstanceID = SourceVertexInstanceIDs[ContourIndex];
+			int32 ContourPlusOne = (ContourIndex + 1) % ContourCount;
+			FVertexInstanceID SourceVertexInstanceID_Next = SourceVertexInstanceIDs[ContourPlusOne];
+
+			int32 ContourPointIndex = Contours.AddDefaulted();
+			FMeshDescription::FContourPoint& ContourPoint = Contours[ContourPointIndex];
 
-		// Check whether or not we have pad the UVs
-		if ( NumTotalUVs < InTarget.WedgeTangentZ.Num() && (InSource.WedgeTexCoords[i].Num() || InTarget.WedgeTexCoords[i].Num()) )
+			const FEdgeID SourceEdgeID = InSource.GetVertexPairEdge(InSource.GetVertexInstanceVertex(SourceVertexInstanceID), InSource.GetVertexInstanceVertex(SourceVertexInstanceID_Next));
+			check(SourceEdgeID != FEdgeID::Invalid);
+			ContourPoint.VertexInstanceID = SourceToTargetVertexInstanceID[SourceVertexInstanceID];
+			ContourPoint.EdgeID = SourceToTargetEdgeID[SourceEdgeID];
+		}
+	};
+
+	//Append Polygons
+	for (const FPolygonID& SourcePolygonID : InSource.Polygons().GetElementIDs())
+	{
+		const FMeshPolygon& SourcePolygon = InSource.GetPolygon(SourcePolygonID);
+		const TArray<FVertexInstanceID>& SourceVertexInstanceIDs = InSource.GetPolygonPerimeterVertexInstances(SourcePolygonID);
+
+
+		TArray<FMeshDescription::FContourPoint> Contours;
+		CreateContour(SourceVertexInstanceIDs, Contours);
+
+		int32 HoleNumber = InSource.GetNumPolygonHoles(SourcePolygonID);
+		TArray<TArray<FMeshDescription::FContourPoint>> AllHoles;
+		AllHoles.AddDefaulted(HoleNumber);
+		for (int32 HoleIndex = 0; HoleIndex < HoleNumber; ++HoleIndex)
 		{
-			InTarget.WedgeTexCoords[i].AddZeroed(InTarget.WedgeTangentZ.Num() - NumTotalUVs);
+			const FMeshPolygonContour& HoleContour = SourcePolygon.HoleContours[HoleIndex];
+			
+			TArray<FMeshDescription::FContourPoint>& Holes = AllHoles[HoleIndex];
+			CreateContour(HoleContour.VertexInstanceIDs, Holes);
 		}
-		
-		InTarget.WedgeTexCoords[i].Append(InSource.WedgeTexCoords[i]);
-	}
 
-	checkf(InTarget.IsValidOrFixable(), TEXT("RawMesh became corrupt after appending InSource"));
+		// Insert a polygon into the mesh
+		const FPolygonID TargetPolygonID = InTarget.CreatePolygon(SourcePolygon.PolygonGroupID, Contours, AllHoles);
+		//Triangulate the polygon
+		FMeshPolygon& Polygon = InTarget.GetPolygon(TargetPolygonID);
+		InTarget.ComputePolygonTriangulation(TargetPolygonID, Polygon.Triangles);
+	}
 }
 
 
-void FMeshMergeHelpers::MergeImpostersToRawMesh(TArray<const UStaticMeshComponent*> ImposterComponents, FRawMesh& InRawMesh, const FVector& InPivot, int32 InBaseMaterialIndex, TArray<UMaterialInterface*>& OutImposterMaterials)
+void FMeshMergeHelpers::MergeImpostersToRawMesh(TArray<const UStaticMeshComponent*> ImposterComponents, FMeshDescription& InRawMesh, const FVector& InPivot, int32 InBaseMaterialIndex, TArray<UMaterialInterface*>& OutImposterMaterials)
 {
 	// TODO decide whether we want this to be user specified or derived from the RawMesh
 	/*const int32 UVOneIndex = [RawMesh, Data]() -> int32
@@ -1184,79 +1368,41 @@ void FMeshMergeHelpers::MergeImpostersToRawMesh(TArray<const UStaticMeshComponen
 			// Retrieve imposter LOD mesh and material			
 			const int32 LODIndex = Component->GetStaticMesh()->GetNumLODs() - 1;
 
-			// Retrieve mesh data in FRawMesh form
-			FRawMesh ImposterMesh;
+			// Retrieve mesh data in FMeshDescription form
+			FMeshDescription ImposterMesh;
 			FMeshMergeHelpers::RetrieveMesh(Component, LODIndex, ImposterMesh, false);
 
 			// Retrieve the sections, we're expect 1 for imposter meshes
 			TArray<FSectionInfo> Sections;
 			FMeshMergeHelpers::ExtractSections(Component, LODIndex, Sections);
 
-			// Generate a map of section to material index remaps
-			TMap<int32, int32> Remaps;
 			for (FSectionInfo& Info : Sections)
 			{
-				const int32 MaterialIndex = OutImposterMaterials.AddUnique(Info.Material);
-				// Offsetting material index by InBaseMaterialIndex
-				Remaps.Add(Info.MaterialIndex, MaterialIndex + InBaseMaterialIndex);
-			}
-
-			// Apply material remapping
-			for (int32& Index : ImposterMesh.FaceMaterialIndices)
-			{
-				Index = Remaps[Index];
+				OutImposterMaterials.AddUnique(Info.Material);
 			}
 
 			// Imposter magic, we're storing the actor world position and X scale spread across two UV channels
 			const int32 UVTwoIndex = UVOneIndex + 1;
-			const int32 NumIndices = ImposterMesh.WedgeIndices.Num();
-			ImposterMesh.WedgeTexCoords[UVOneIndex].SetNumZeroed(NumIndices);
-			ImposterMesh.WedgeTexCoords[UVTwoIndex].SetNumZeroed(NumIndices);
-
+			TVertexInstanceAttributesRef<FVector2D> VertexInstanceUVs = ImposterMesh.VertexInstanceAttributes().GetAttributesRef<FVector2D>(MeshAttribute::VertexInstance::TextureCoordinate);
+			VertexInstanceUVs.SetNumIndices(UVTwoIndex + 1);
+			const int32 NumIndices = ImposterMesh.VertexInstances().Num();
 			const FTransform& ActorToWorld = Component->GetOwner()->GetActorTransform();
 			const FVector ActorPosition = ActorToWorld.TransformPosition(FVector::ZeroVector) - InPivot;
-			for (int32 Index = 0; Index < NumIndices; ++Index)
+			for(const FVertexInstanceID& VertexInstanceID : ImposterMesh.VertexInstances().GetElementIDs())
 			{
-				const int32& VertexIndex = ImposterMesh.WedgeIndices[Index];
-
-				const FVector& WedgePosition = ImposterMesh.VertexPositions[VertexIndex];
-
-				FVector2D& UVOne = ImposterMesh.WedgeTexCoords[UVOneIndex][Index];
-				FVector2D& UVTwo = ImposterMesh.WedgeTexCoords[UVTwoIndex][Index];
+				FVector2D UVOne;
+				FVector2D UVTwo;
 					
 				UVOne.X = ActorPosition.X;
 				UVOne.Y = ActorPosition.Y;
+				VertexInstanceUVs.Set(VertexInstanceID, UVOneIndex, UVOne);
 
 				UVTwo.X = ActorPosition.Z;
 				UVTwo.Y = ActorToWorld.GetScale3D().X;
+				VertexInstanceUVs.Set(VertexInstanceID, UVTwoIndex, UVTwo);
 			}
 
 			FMeshMergeHelpers::AppendRawMesh(InRawMesh, ImposterMesh);
 		}
 	}
 }
-
-FMeshDescription* FMeshMergeData::GetMeshDescription()
-{
-	if (MeshDescription != nullptr)
-	{
-		return MeshDescription;
-	}
-	if (RawMesh == nullptr)
-	{
-		return nullptr;
-	}
-
-	MeshDescription = new FMeshDescription();
-	UStaticMesh::RegisterMeshAttributes(*MeshDescription);
-	TMap<int32, FName> MaterialMap;
-	if (SourceStaticMesh)
-	{
-		for (int32 MaterialIndex = 0; MaterialIndex < SourceStaticMesh->StaticMaterials.Num(); ++MaterialIndex)
-		{
-			MaterialMap.Add(MaterialIndex, SourceStaticMesh->StaticMaterials[MaterialIndex].ImportedMaterialSlotName);
-		}
-	}
-	FMeshDescriptionOperations::ConvertFromRawMesh(*RawMesh, *MeshDescription, MaterialMap);
-	return MeshDescription;
-}
\ No newline at end of file
diff --git a/Engine/Source/Developer/MeshMergeUtilities/Private/MeshMergeHelpers.h b/Engine/Source/Developer/MeshMergeUtilities/Private/MeshMergeHelpers.h
--- a/Engine/Source/Developer/MeshMergeUtilities/Private/MeshMergeHelpers.h
+++ b/Engine/Source/Developer/MeshMergeUtilities/Private/MeshMergeHelpers.h
@@ -12,7 +12,8 @@ class USplineMeshComponent;
 class UBodySetup;
 class ALandscapeProxy;
 struct FSectionInfo;
-struct FRawMesh;
+struct FMeshDescription;
+struct FStaticMaterial;
 struct FRawMeshExt;
 struct FStaticMeshLODResources;
 struct FKAggregateGeom;
@@ -27,33 +28,34 @@ public:
 	static void ExtractSections(const UStaticMesh* StaticMesh, int32 LODIndex, TArray<FSectionInfo>& OutSections);
 
 	/** Expanding instance data from instanced static mesh components */
-	static void ExpandInstances(const UInstancedStaticMeshComponent* InInstancedStaticMeshComponent, FRawMesh& InOutRawMesh, TArray<FSectionInfo>& InOutSections);
+	static void ExpandInstances(const UInstancedStaticMeshComponent* InInstancedStaticMeshComponent, FMeshDescription& InOutRawMesh, TArray<FSectionInfo>& InOutSections);
 
-	/** Extracting mesh data in FRawMesh form from static, skeletal mesh (components) */
-	static void RetrieveMesh(const UStaticMeshComponent* StaticMeshComponent, int32 LODIndex, FRawMesh& RawMesh, bool bPropagateVertexColours);
-	static void RetrieveMesh(USkeletalMeshComponent* SkeletalMeshComponent, int32 LODIndex, FRawMesh& RawMesh, bool bPropagateVertexColours);
-	static void RetrieveMesh(const UStaticMesh* StaticMesh, int32 LODIndex, FRawMesh& RawMesh);
+	/** Extracting mesh data in FMeshDescription form from static, skeletal mesh (components) */
+	static void RetrieveMesh(const UStaticMeshComponent* StaticMeshComponent, int32 LODIndex, FMeshDescription& RawMesh, bool bPropagateVertexColours);
+	static void RetrieveMesh(USkeletalMeshComponent* SkeletalMeshComponent, int32 LODIndex, FMeshDescription& RawMesh, bool bPropagateVertexColours);
+	static void RetrieveMesh(const UStaticMesh* StaticMesh, int32 LODIndex, FMeshDescription& RawMesh);
 	
 	/** Exports static mesh LOD render data to a RawMesh */
-	static void ExportStaticMeshLOD(const FStaticMeshLODResources& StaticMeshLOD, FRawMesh& OutRawMesh);
+	static void ExportStaticMeshLOD(const FStaticMeshLODResources& StaticMeshLOD, FMeshDescription& OutRawMesh, const TArray<FStaticMaterial>& Materials);
 
 	/** Checks whether or not the texture coordinates are outside of 0-1 UV ranges */
 	static bool CheckWrappingUVs(const TArray<FVector2D>& UVs);
+	static bool CheckWrappingUVs(const FMeshDescription& MeshDescription, int32 UVChannelIndex);
 
 	/** Culls away triangles which are inside culling volumes or completely underneath the landscape */
-	static void CullTrianglesFromVolumesAndUnderLandscapes(const UWorld* World, const FBoxSphereBounds& Bounds, FRawMesh& InOutRawMesh);
+	static void CullTrianglesFromVolumesAndUnderLandscapes(const UWorld* World, const FBoxSphereBounds& Bounds, FMeshDescription& InOutRawMesh);
 	
 	/** Propagates deformation along spline to raw mesh data */
-	static void PropagateSplineDeformationToRawMesh(const USplineMeshComponent* InSplineMeshComponent, struct FRawMesh &OutRawMesh);
+	static void PropagateSplineDeformationToRawMesh(const USplineMeshComponent* InSplineMeshComponent, FMeshDescription &OutRawMesh);
 	
 	/** Propagates deformation along spline to physics geometry data */
 	static void PropagateSplineDeformationToPhysicsGeometry(USplineMeshComponent* SplineMeshComponent, FKAggregateGeom& InOutPhysicsGeometry);
 
 	/** Transforms raw mesh data using InTransform*/
-	static void TransformRawMeshVertexData(const FTransform& InTransform, FRawMesh &OutRawMesh);
+	static void TransformRawMeshVertexData(const FTransform& InTransform, FMeshDescription &OutRawMesh);
 
-	/** Retrieves all culling landscapes and volumes as FRawMesh structures */
-	static void RetrieveCullingLandscapeAndVolumes(UWorld* InWorld, const FBoxSphereBounds& EstimatedMeshProxyBounds, const TEnumAsByte<ELandscapeCullingPrecision::Type> PrecisionType, TArray<FRawMesh*>& CullingRawMeshes);
+	/** Retrieves all culling landscapes and volumes as FMeshDescription structures */
+	static void RetrieveCullingLandscapeAndVolumes(UWorld* InWorld, const FBoxSphereBounds& EstimatedMeshProxyBounds, const TEnumAsByte<ELandscapeCullingPrecision::Type> PrecisionType, TArray<FMeshDescription*>& CullingRawMeshes);
 
 	/** Transforms physics geometry data using InTransform */
 	static void TransformPhysicsGeometry(const FTransform& InTransform, const bool bBakeConvexTransform, struct FKAggregateGeom& AggGeom);
@@ -65,17 +67,17 @@ public:
 	static FVector2D GetValidUV(const FVector2D& UV);
 	
 	/** Calculates UV coordinates bounds for the given Raw Mesh	*/
-	static void CalculateTextureCoordinateBoundsForRawMesh(const FRawMesh& InRawMesh, TArray<FBox2D>& OutBounds);
+	static void CalculateTextureCoordinateBoundsForRawMesh(const FMeshDescription& InRawMesh, TArray<FBox2D>& OutBounds);
 
 	/** Propagates vertex painted colors from the StaticMeshComponent instance to RawMesh */
-	static bool PropagatePaintedColorsToRawMesh(const UStaticMeshComponent* StaticMeshComponent, int32 LODIndex, FRawMesh& RawMesh);
+	static bool PropagatePaintedColorsToRawMesh(const UStaticMeshComponent* StaticMeshComponent, int32 LODIndex, FMeshDescription& RawMesh);
 
 	/** Checks whether or not the landscape proxy is hit given a ray start and end */
 	static bool IsLandscapeHit(const FVector& RayOrigin, const FVector& RayEndPoint, const UWorld* World, const TArray<ALandscapeProxy*>& LandscapeProxies, FVector& OutHitLocation);
 	
-	/** Appends a FRawMesh to another instance */
-	static void AppendRawMesh(FRawMesh& InTarget, const FRawMesh& InSource);
+	/** Appends a FMeshDescription to another instance */
+	static void AppendRawMesh(FMeshDescription& InTarget, const FMeshDescription& InSource);
 
 	/** Merges imposter meshes into a raw mesh. */
-	static void MergeImpostersToRawMesh(TArray<const UStaticMeshComponent*> ImposterComponents, FRawMesh& InRawMesh, const FVector& InPivot, int32 BaseMaterialIndex, TArray<UMaterialInterface*>& OutImposterMaterials);
+	static void MergeImpostersToRawMesh(TArray<const UStaticMeshComponent*> ImposterComponents, FMeshDescription& InRawMesh, const FVector& InPivot, int32 BaseMaterialIndex, TArray<UMaterialInterface*>& OutImposterMaterials);
 };
diff --git a/Engine/Source/Developer/MeshMergeUtilities/Private/MeshMergeUtilities.cpp b/Engine/Source/Developer/MeshMergeUtilities/Private/MeshMergeUtilities.cpp
--- a/Engine/Source/Developer/MeshMergeUtilities/Private/MeshMergeUtilities.cpp
+++ b/Engine/Source/Developer/MeshMergeUtilities/Private/MeshMergeUtilities.cpp
@@ -7,7 +7,6 @@
 
 #include "MaterialOptions.h"
 #include "IMaterialBakingModule.h"
-#include "RawMesh.h"
 
 #include "Misc/PackageName.h"
 #include "MaterialUtilities.h"
@@ -68,6 +67,12 @@
 #include "MaterialBakingHelpers.h"
 #include "IMeshMergeExtension.h"
 
+#include "RawMesh.h"
+#include "MeshDescription.h"
+#include "MeshAttributes.h"
+#include "MeshAttributeArray.h"
+#include "MeshDescriptionOperations.h"
+
 #if WITH_EDITOR
 #include "Widgets/Notifications/SNotificationList.h"
 #include "Framework/Notifications/NotificationManager.h"
@@ -111,7 +116,7 @@ void FMeshMergeUtilities::BakeMaterialsForComponent(TArray<TWeakObjectPtr<UObjec
 	UMaterialMergeOptions* MaterialMergeOptions  = MaterialMergeOptionsObject ? Cast<UMaterialMergeOptions>(MaterialMergeOptionsObject->Get()) : nullptr;
 
 	// Mesh / LOD index	
-	TMap<uint32, FRawMesh> RawMeshLODs;
+	TMap<uint32, FMeshDescription> RawMeshLODs;
 
 	// LOD index, <original section index, unique section index>
 	TMultiMap<uint32, TPair<uint32, uint32>> UniqueSectionIndexPerLOD;
@@ -133,7 +138,7 @@ void FMeshMergeUtilities::BakeMaterialsForComponent(TArray<TWeakObjectPtr<UObjec
 		const bool bProcessedLOD = MaterialOptions->LODIndices.Contains(LODIndex);
 		if (bProcessedLOD)
 		{
-			FRawMesh& RawMesh = RawMeshLODs.Add(LODIndex);
+			FMeshDescription& RawMesh = RawMeshLODs.Add(LODIndex);
 			Adapter->RetrieveRawMeshData(LODIndex, RawMesh, MaterialOptions->bUseMeshData);
 		}
 
@@ -195,7 +200,10 @@ void FMeshMergeUtilities::BakeMaterialsForComponent(TArray<TWeakObjectPtr<UObjec
 				if (MeshSettings.MaterialIndices.Num())
 				{
 					// Retrieve raw mesh
-					MeshSettings.RawMesh = RawMeshLODs.Find(LODIndex);
+					const FMeshDescription* MeshDescription = RawMeshLODs.Find(LODIndex);
+					TMap<FName, int32> MaterialMap;
+					FMeshDescriptionOperations::ConvertToRawMesh(*MeshDescription, *(MeshSettings.RawMesh), MaterialMap);
+					
 
 					MeshSettings.TextureCoordinateBox = FBox2D(FVector2D(0.0f, 0.0f), FVector2D(1.0f, 1.0f));
 					const bool bUseVertexColor = (MeshSettings.RawMesh->WedgeColors.Num() > 0);
@@ -1154,9 +1162,9 @@ void FMeshMergeUtilities::CreateProxyMesh(const TArray<UStaticMeshComponent*>& I
 	SlowTask.EnterProgressFrame(5.0f, LOCTEXT("CreateProxyMesh_CollectingMeshes", "Collecting Input Static Meshes"));
 
 	// Mesh / LOD index	
-	TMap<uint32, FRawMesh*> RawMeshLODs;
+	TMap<uint32, FMeshDescription*> RawMeshLODs;
 
-	TArray<FRawMesh*> RawMeshData;
+	TArray<FMeshDescription*> RawMeshData;
 
 	// LOD index, <original section index, unique section index>
 	TMultiMap<uint32, TPair<uint32, uint32>> UniqueSectionIndexPerLOD;
@@ -1183,8 +1191,9 @@ void FMeshMergeUtilities::CreateProxyMesh(const TArray<UStaticMeshComponent*>& I
 			const int32 LODIndex = InMeshProxySettings.bCalculateCorrectLODModel ? ScreenSizeBasedLODLevel : 0;
 			static const bool bPropagateVertexColours = true;
 
-			// Retrieve mesh data in FRawMesh form
-			FRawMesh* RawMesh = new FRawMesh();
+			// Retrieve mesh data in FMeshDescription form
+			FMeshDescription* RawMesh = new FMeshDescription();
+			UStaticMesh::RegisterMeshAttributes(*RawMesh);
 			FMeshMergeHelpers::RetrieveMesh(StaticMeshComponent, LODIndex, *RawMesh, bPropagateVertexColours);
 			const int32 MeshIndex = RawMeshData.Add(RawMesh);
 
@@ -1274,8 +1283,18 @@ void FMeshMergeUtilities::CreateProxyMesh(const TArray<UStaticMeshComponent*>& I
 				for (uint32 MeshIndex : MeshIndices)
 				{
 					FMeshData MeshSettings;
-					// Add entries for each used mesh
-					MeshSettings.RawMesh = RawMeshData[MeshIndex];
+					// Retrieve raw mesh
+					FMeshDescription* MeshDescription = RawMeshData[MeshIndex];
+					TMap<FName, int32> MaterialMap;
+					int32 MeshMaterialIndex = 0;
+					TPolygonGroupAttributesConstRef<FName> TargetPolygonGroupImportedMaterialSlotNames = MeshDescription->PolygonGroupAttributes().GetAttributesRef<FName>(MeshAttribute::PolygonGroup::ImportedMaterialSlotName);
+					for (const FPolygonGroupID& PolygonGroupID : MeshDescription->PolygonGroups().GetElementIDs())
+					{
+						MaterialMap.Add(TargetPolygonGroupImportedMaterialSlotNames[PolygonGroupID], MeshMaterialIndex);
+						MeshMaterialIndex++;
+					}
+					FMeshDescriptionOperations::ConvertToRawMesh(*MeshDescription, *(MeshSettings.RawMesh), MaterialMap);
+					MeshSettings.RawMeshDescription = MeshDescription;
 
 					// If we already have lightmap uvs generated or the lightmap coordinate index != 0 and available we can reuse those instead of having to generate new ones
 					if (InMeshProxySettings.bReuseMeshLightmapUVs && (ComponentsToMerge[MeshIndex]->GetStaticMesh()->SourceModels[0].BuildSettings.bGenerateLightmapUVs || (ComponentsToMerge[MeshIndex]->GetStaticMesh()->LightMapCoordinateIndex != 0 && MeshSettings.RawMesh->WedgeTexCoords[ComponentsToMerge[MeshIndex]->GetStaticMesh()->LightMapCoordinateIndex].Num() != 0)))
@@ -1333,6 +1352,7 @@ void FMeshMergeUtilities::CreateProxyMesh(const TArray<UStaticMeshComponent*>& I
 			// Add simple bake entry 
 			FMeshData MeshSettings;
 			MeshSettings.RawMesh = nullptr;
+			MeshSettings.RawMeshDescription = nullptr;
 			MeshSettings.TextureCoordinateBox = FBox2D(FVector2D(0.0f, 0.0f), FVector2D(1.0f, 1.0f));
 			MeshSettings.TextureCoordinateIndex = 0;
 
@@ -1405,7 +1425,7 @@ void FMeshMergeUtilities::CreateProxyMesh(const TArray<UStaticMeshComponent*>& I
 		// Now have the baked out material data, need to have a map or actually remap the raw mesh data to baked material indices
 		for (int32 MeshIndex = 0; MeshIndex < RawMeshData.Num(); ++MeshIndex)
 		{
-			FRawMesh& RawMesh = *RawMeshData[MeshIndex];
+			FMeshDescription& RawMesh = *RawMeshData[MeshIndex];
 
 			TArray<TPair<uint32, uint32>> SectionAndOutputIndices;
 			OutputMaterialsMap.MultiFind(MeshIndex, SectionAndOutputIndices);
@@ -1424,17 +1444,45 @@ void FMeshMergeUtilities::CreateProxyMesh(const TArray<UStaticMeshComponent*>& I
 
 				Remap[SectionIndex] = NewIndex;
 			}
+			
+			TMap<FPolygonGroupID, FPolygonGroupID> RemapPolygonGroup;
+			int32 MaxRemapID = 0;
+			for (const FPolygonGroupID& PolygonGroupID : RawMesh.PolygonGroups().GetElementIDs())
+			{
+				checkf(Remap.IsValidIndex(PolygonGroupID.GetValue()), TEXT("Missing material bake output index entry for mesh(section)"));
+				int32 RemapID = Remap[PolygonGroupID.GetValue()];
+				RemapPolygonGroup.Add(PolygonGroupID, FPolygonGroupID(RemapID));
+				MaxRemapID = FMath::Max(MaxRemapID, RemapID);
+			}
+
+			TSparseArray<int32> PolygonGroupRemap;
+			PolygonGroupRemap.Reserve(MaxRemapID);
+			for (int32 RemapID = 0; RemapID < MaxRemapID; ++RemapID)
+			{
+				PolygonGroupRemap.AddUninitialized();
+			}
+			//Set the polygon state
+			for (auto Kvp : RemapPolygonGroup)
+			{
+				PolygonGroupRemap[Kvp.Key.GetValue()] = Kvp.Value.GetValue();
+				FMeshPolygonGroup& PolygonGroup = RawMesh.GetPolygonGroup(Kvp.Key);
+				for (FPolygonID PolygonID : PolygonGroup.Polygons)
+				{
+					FMeshPolygon& Polygon = RawMesh.GetPolygon(PolygonID);
+					Polygon.PolygonGroupID = Kvp.Value;
+				}
+			}
 
-			for (int32& FaceMaterialIndex : RawMesh.FaceMaterialIndices)
+			//Remap the polygon groups
+			for (auto Kvp : RemapPolygonGroup)
 			{
-				checkf(Remap.IsValidIndex(FaceMaterialIndex), TEXT("Missing material bake output index entry for mesh(section)"));
-				FaceMaterialIndex = Remap[FaceMaterialIndex];
+				RawMesh.PolygonGroups().Remap(PolygonGroupRemap);
 			}
 		}
 	};
 
 	// Landscape culling
-	TArray<FRawMesh*> CullingRawMeshes;
+	TArray<FMeshDescription*> CullingRawMeshes;
 	if (InMeshProxySettings.bUseLandscapeCulling)
 	{
 		SlowTask.EnterProgressFrame(5.0f, LOCTEXT("CreateProxyMesh_LandscapeCulling", "Applying Landscape Culling"));
@@ -1475,7 +1523,7 @@ void FMeshMergeUtilities::CreateProxyMesh(const TArray<UStaticMeshComponent*>& I
 
 		FMeshData* MeshData = GlobalMeshSettings.FindByPredicate([&](const FMeshData& Entry)
 		{
-			return Entry.RawMesh == MergeData.RawMesh && (Entry.CustomTextureCoordinates.Num() || Entry.TextureCoordinateIndex != 0);
+			return Entry.RawMeshDescription == MergeData.RawMesh && (Entry.CustomTextureCoordinates.Num() || Entry.TextureCoordinateIndex != 0);
 		});
 
 		if (MeshData)
@@ -1494,7 +1542,7 @@ void FMeshMergeUtilities::CreateProxyMesh(const TArray<UStaticMeshComponent*>& I
 	}
 
 	// Populate landscape clipping geometry
-	for (FRawMesh* RawMesh : CullingRawMeshes)
+	for (FMeshDescription* RawMesh : CullingRawMeshes)
 	{
 		FMeshMergeData ClipData;
 		ClipData.bIsClippingMesh = true;
@@ -1620,6 +1668,82 @@ void FMeshMergeUtilities::UnregisterExtension(IMeshMergeExtension* InExtension)
 	MeshMergeExtensions.Remove(InExtension);
 }
 
+bool RetrieveRawMeshData(FMeshMergeDataTracker& DataTracker
+	, const int32 ComponentIndex
+	, const int32 LODIndex
+	, UStaticMeshComponent* Component
+	, const bool bPropagateMeshData
+	, TArray<FSectionInfo>& Sections
+	, FStaticMeshComponentAdapter& Adapter
+	, const bool bMergeMaterialData
+	, const FMeshMergingSettings& InSettings)
+{
+	// Retrieve raw mesh data
+	FMeshDescription& RawMesh = DataTracker.AddAndRetrieveRawMesh(ComponentIndex, LODIndex, Component->GetStaticMesh());
+	Adapter.RetrieveRawMeshData(LODIndex, RawMesh, bPropagateMeshData);
+
+	// Reset section for reuse
+	Sections.SetNum(0, false);
+
+	// Extract sections for given LOD index from the mesh 
+	Adapter.RetrieveMeshSections(LODIndex, Sections);
+
+	TMap<FPolygonGroupID, FPolygonGroupID> RemapPolygonGroup;
+	for (int32 SectionIndex = 0; SectionIndex < Sections.Num(); ++SectionIndex)
+	{
+		const FSectionInfo& Section = Sections[SectionIndex];
+		// Unique section index for remapping
+		const int32 UniqueIndex = DataTracker.AddSection(Section);
+
+		// Store of original to unique section index entry for this component + LOD index
+		DataTracker.AddSectionRemapping(ComponentIndex, LODIndex, Section.MaterialIndex, UniqueIndex);
+		DataTracker.AddMaterialSlotName(Section.Material, Section.MaterialSlotName);
+
+		if (!bMergeMaterialData)
+		{
+			FMeshDescriptionOperations::SwapPolygonPolygonGroup(RawMesh, UniqueIndex, Section.StartIndex, Section.EndIndex, false);
+		}
+		else
+		{
+			RemapPolygonGroup.Add(FPolygonGroupID(SectionIndex), FPolygonGroupID(Section.MaterialIndex));
+		}
+	}
+	
+	if (bMergeMaterialData)
+	{
+		FMeshDescriptionOperations::RemapPolygonGroups(RawMesh, RemapPolygonGroup);
+	}
+	
+	//Compact the PolygonGroupID to make sure it follow the section index
+	FElementIDRemappings RemapInformation;
+	RawMesh.Compact(RemapInformation);
+
+	// If the component is an ISMC then we need to duplicate the vertex data
+	if (Component->GetClass() == UInstancedStaticMeshComponent::StaticClass())
+	{
+		const UInstancedStaticMeshComponent* InstancedStaticMeshComponent = Cast<UInstancedStaticMeshComponent>(Component);
+		FMeshMergeHelpers::ExpandInstances(InstancedStaticMeshComponent, RawMesh, Sections);
+	}
+
+	if (InSettings.bUseLandscapeCulling)
+	{
+		FMeshMergeHelpers::CullTrianglesFromVolumesAndUnderLandscapes(Component->GetWorld(), Adapter.GetBounds(), RawMesh);
+	}
+
+	// If the valid became invalid during retrieval remove it again
+	const bool bValidMesh = RawMesh.VertexInstances().Num() > 0;
+	if (!bValidMesh)
+	{
+		DataTracker.RemoveRawMesh(ComponentIndex, LODIndex);
+	}
+	else if (Component->GetStaticMesh() != nullptr)
+	{
+		// If the mesh is valid at this point, record the lightmap UV so we have a record for use later
+		DataTracker.AddLightmapChannelRecord(ComponentIndex, LODIndex, Component->GetStaticMesh()->LightMapCoordinateIndex);
+	}
+	return bValidMesh;
+}
+
 void FMeshMergeUtilities::MergeComponentsToStaticMesh(const TArray<UPrimitiveComponent*>& ComponentsToMerge, UWorld* World, const FMeshMergingSettings& InSettings, UMaterialInterface* InBaseMaterial, UPackage* InOuter, const FString& InBasePackageName, TArray<UObject*>& OutAssetsToSync, FVector& OutMergedActorLocation, const float ScreenSize, bool bSilent /*= false*/) const
 {
 	// Use first mesh for naming and pivot
@@ -1702,52 +1826,19 @@ void FMeshMergeUtilities::MergeComponentsToStaticMesh(const TArray<UPrimitiveCom
 
 			for (int32 LODIndex = 0; LODIndex < NumLODs; ++LODIndex)
 			{
-				FRawMesh& RawMesh = DataTracker.AddAndRetrieveRawMesh(ComponentIndex, LODIndex, Component->GetStaticMesh());
-				Adapter.RetrieveRawMeshData(LODIndex, RawMesh, bPropagateMeshData);
-
-				// Reset section for reuse
-				Sections.SetNum(0, false);
-
-				// Extract sections for given LOD index from the mesh 
-				Adapter.RetrieveMeshSections(LODIndex, Sections);
-
-				for (int32 SectionIndex = 0; SectionIndex < Sections.Num(); ++SectionIndex)
-				{
-					const FSectionInfo& Section = Sections[SectionIndex];
-					const int32 UniqueIndex = DataTracker.AddSection(Section);
-					DataTracker.AddSectionRemapping(ComponentIndex, LODIndex, Section.MaterialIndex, UniqueIndex);
-					DataTracker.AddMaterialSlotName(Section.Material, Section.MaterialSlotName);
-					for (int32 StartIndex = Section.StartIndex; StartIndex < Section.EndIndex; ++StartIndex)
-					{
-						RawMesh.FaceMaterialIndices[StartIndex] = UniqueIndex;
-					}
-				}
-
-				// If the component is an ISMC then we need to duplicate the vertex data
-				if(Component->GetClass() == UInstancedStaticMeshComponent::StaticClass())
-				{
-					const UInstancedStaticMeshComponent* InstancedStaticMeshComponent = Cast<UInstancedStaticMeshComponent>(Component);
-					FMeshMergeHelpers::ExpandInstances(InstancedStaticMeshComponent, RawMesh, Sections);
-				}
-
-				if (InSettings.bUseLandscapeCulling)
+				if (!RetrieveRawMeshData(DataTracker
+					, ComponentIndex
+					, LODIndex
+					, Component
+					, bPropagateMeshData
+					, Sections
+					, Adapter
+					, false
+					, InSettings))
 				{
-					FMeshMergeHelpers::CullTrianglesFromVolumesAndUnderLandscapes(Component->GetWorld(), Adapter.GetBounds(), RawMesh);
-				}
-
-				bool bValidMesh = RawMesh.IsValid();
-
-				if (!bValidMesh)
-				{
-					DataTracker.RemoveRawMesh(ComponentIndex, LODIndex);
+					//If the rawmesh was not retrieve properly break the loop
 					break;
 				}
-				else if(Component->GetStaticMesh() != nullptr)
-				{
-					// If the mesh is valid at this point, record the lightmap UV so we have a record for use later
-					DataTracker.AddLightmapChannelRecord(ComponentIndex, LODIndex, Component->GetStaticMesh()->LightMapCoordinateIndex);
-				}
-
 				DataTracker.AddLODIndex(LODIndex);
 			}
 		}
@@ -1799,58 +1890,15 @@ void FMeshMergeUtilities::MergeComponentsToStaticMesh(const TArray<UPrimitiveCom
 				}
 			}();
 
-			// Retrieve raw mesh data
-			FRawMesh& RawMesh = DataTracker.AddAndRetrieveRawMesh(ComponentIndex, LODIndex, Component->GetStaticMesh());
-			Adapter.RetrieveRawMeshData(LODIndex, RawMesh, bPropagateMeshData);
-
-			// Reset section for reuse
-			Sections.SetNum(0, false);
-
-			// Extract sections for given LOD index from the mesh 
-			Adapter.RetrieveMeshSections(LODIndex, Sections);
-
-			for (int32 SectionIndex = 0; SectionIndex < Sections.Num(); ++SectionIndex)
-			{
-				const FSectionInfo& Section = Sections[SectionIndex];
-				// Unique section index for remapping
-				const int32 UniqueIndex = DataTracker.AddSection(Section);
-
-				// Store of original to unique section index entry for this component + LOD index
-				DataTracker.AddSectionRemapping(ComponentIndex, LODIndex, Section.MaterialIndex, UniqueIndex);
-				DataTracker.AddMaterialSlotName(Section.Material, Section.MaterialSlotName);
-
-				if (!bMergeMaterialData)
-				{
-					for (int32 StartIndex = Section.StartIndex; StartIndex < Section.EndIndex; ++StartIndex)
-					{
-						RawMesh.FaceMaterialIndices[StartIndex] = UniqueIndex;
-					}
-				}
-			}
-
-			// If the component is an ISMC then we need to duplicate the vertex data
-			if(Component->GetClass() == UInstancedStaticMeshComponent::StaticClass())
-			{
-				const UInstancedStaticMeshComponent* InstancedStaticMeshComponent = Cast<UInstancedStaticMeshComponent>(Component);
-				FMeshMergeHelpers::ExpandInstances(InstancedStaticMeshComponent, RawMesh, Sections);
-			}
-
-			if (InSettings.bUseLandscapeCulling)
-			{
-				FMeshMergeHelpers::CullTrianglesFromVolumesAndUnderLandscapes(Component->GetWorld(), Adapter.GetBounds(), RawMesh);
-			}
-
-			// If the valid became invalid during retrieval remove it again
-			const bool bValidMesh = RawMesh.IsValid();
-			if (!bValidMesh)
-			{
-				DataTracker.RemoveRawMesh(ComponentIndex, LODIndex);
-			}
-			else if(Component->GetStaticMesh() != nullptr)
-			{
-				// If the mesh is valid at this point, record the lightmap UV so we have a record for use later
-				DataTracker.AddLightmapChannelRecord(ComponentIndex, LODIndex, Component->GetStaticMesh()->LightMapCoordinateIndex);
-			}
+			RetrieveRawMeshData(DataTracker
+				, ComponentIndex
+				, LODIndex
+				, Component
+				, bPropagateMeshData
+				, Sections
+				, Adapter
+				, bMergeMaterialData
+				, InSettings);
 		}
 	}
 
@@ -1928,7 +1976,7 @@ void FMeshMergeUtilities::MergeComponentsToStaticMesh(const TArray<UPrimitiveCom
 		for (TConstRawMeshIterator RawMeshIterator = DataTracker.GetConstRawMeshIterator(); RawMeshIterator; ++RawMeshIterator)
 		{
 			const FMeshLODKey& Key = RawMeshIterator.Key();
-			const FRawMesh& RawMesh = RawMeshIterator.Value();
+			const FMeshDescription& RawMesh = RawMeshIterator.Value();
 			const bool bRequiresUniqueUVs = DataTracker.DoesMeshLODRequireUniqueUVs(Key);
 			UStaticMeshComponent* Component = StaticMeshComponentsToMerge[Key.GetMeshIndex()];
 
@@ -1969,10 +2017,14 @@ void FMeshMergeUtilities::MergeComponentsToStaticMesh(const TArray<UPrimitiveCom
 
 				if (InSettings.bUseVertexDataForBakingMaterial && (bDoesMaterialUseVertexData || bRequiresUniqueUVs))
 				{
-					MeshData.RawMesh = DataTracker.GetRawMeshPtr(Key);
+					MeshData.RawMeshDescription = DataTracker.GetRawMeshPtr(Key);
+					TMap<FName, int32> MaterialMap;
+					MeshData.RawMesh = new FRawMesh();
+					FMeshDescriptionOperations::ConvertToRawMesh(*MeshData.RawMeshDescription, *(MeshData.RawMesh), MaterialMap);
+
 					// if it has vertex color/*WedgetColors.Num()*/, it should also use light map UV index
 					// we can't do this for all meshes, but only for the mesh that has vertex color.
-					if (bRequiresUniqueUVs || MeshData.RawMesh->WedgeColors.Num() > 0)
+					if (bRequiresUniqueUVs || MeshData.RawMeshDescription->VertexInstances().Num() > 0)
 					{
 						// Check if there are lightmap uvs available?
 						const int32 LightMapUVIndex = StaticMeshComponentsToMerge[Key.GetMeshIndex()]->GetStaticMesh()->LightMapCoordinateIndex;
@@ -1985,8 +2037,7 @@ void FMeshMergeUtilities::MergeComponentsToStaticMesh(const TArray<UPrimitiveCom
 						{
 							if (!UniqueTextureCoordinates.Num())
 							{
-								IMeshUtilities& MeshUtilities = FModuleManager::LoadModuleChecked<IMeshUtilities>("MeshUtilities");
-								MeshUtilities.GenerateUniqueUVsForStaticMesh(*MeshData.RawMesh, MaterialOptions->TextureSize.GetMax(), UniqueTextureCoordinates);
+								FMeshDescriptionOperations::GenerateUniqueUVsForStaticMesh(*MeshData.RawMeshDescription, MaterialOptions->TextureSize.GetMax(), false, UniqueTextureCoordinates);
 								ScaleTextureCoordinatesToBox(FBox2D(FVector2D::ZeroVector, FVector2D(1, 1)), UniqueTextureCoordinates);
 							}
 							MeshData.CustomTextureCoordinates = UniqueTextureCoordinates;
@@ -2044,6 +2095,7 @@ void FMeshMergeUtilities::MergeComponentsToStaticMesh(const TArray<UPrimitiveCom
 				else
 				{
 					MeshData.RawMesh = nullptr;
+					MeshData.RawMeshDescription = nullptr;
 					MeshData.TextureCoordinateBox = FBox2D(FVector2D(0.0f, 0.0f), FVector2D(1.0f, 1.0f));
 
 					// This prevents baking out the same material multiple times, which would be wasteful when it does not use vertex data anyway
@@ -2085,13 +2137,12 @@ void FMeshMergeUtilities::MergeComponentsToStaticMesh(const TArray<UPrimitiveCom
 		const bool bGloballyRemapUVs = !bMergeAllLODs && !InSettings.bReuseMeshLightmapUVs;
 		if(bGloballyRemapUVs)
 		{
-			TArray<FRawMesh> MergedRawMeshes;
+			TArray<FMeshDescription> MergedRawMeshes;
 			CreateMergedRawMeshes(DataTracker, InSettings, StaticMeshComponentsToMerge, UniqueMaterials, CollapsedMaterialMap, OutputMaterialsMap, false, false, MergedAssetPivot, MergedRawMeshes);
 
 			// Create texture coords for the merged mesh
 			TArray<FVector2D> GlobalTextureCoordinates;
-			IMeshUtilities& MeshUtilities = FModuleManager::LoadModuleChecked<IMeshUtilities>("MeshUtilities");
-			MeshUtilities.GenerateUniqueUVsForStaticMesh(MergedRawMeshes[0], MaterialOptions->TextureSize.GetMax(), true, GlobalTextureCoordinates);
+			FMeshDescriptionOperations::GenerateUniqueUVsForStaticMesh(MergedRawMeshes[0], MaterialOptions->TextureSize.GetMax(), true, GlobalTextureCoordinates);
 			ScaleTextureCoordinatesToBox(FBox2D(FVector2D::ZeroVector, FVector2D(1, 1)), GlobalTextureCoordinates);
 
 			// copy UVs back to the un-merged mesh's custom texture coords
@@ -2100,11 +2151,11 @@ void FMeshMergeUtilities::MergeComponentsToStaticMesh(const TArray<UPrimitiveCom
 			for (TConstRawMeshIterator RawMeshIterator = DataTracker.GetConstRawMeshIterator(); RawMeshIterator; ++RawMeshIterator)
 			{
 				const FMeshLODKey& Key = RawMeshIterator.Key();
-				const FRawMesh& RawMesh = RawMeshIterator.Value();
+				const FMeshDescription& RawMesh = RawMeshIterator.Value();
 
 				// Build a local array for this raw mesh
 				TArray<FVector2D> UniqueTextureCoordinates;
-				UniqueTextureCoordinates.SetNumUninitialized(RawMesh.WedgeIndices.Num());
+				UniqueTextureCoordinates.SetNumUninitialized(RawMesh.VertexInstances().Num());
 				for(FVector2D& UniqueTextureCoordinate : UniqueTextureCoordinates)
 				{
 					UniqueTextureCoordinate = GlobalTextureCoordinates[GlobalUVIndex++];
@@ -2113,7 +2164,7 @@ void FMeshMergeUtilities::MergeComponentsToStaticMesh(const TArray<UPrimitiveCom
 				// copy to mesh data
 				for(FMeshData& MeshData : GlobalMeshSettings)
 				{
-					if(MeshData.RawMesh == &RawMesh)
+					if(MeshData.RawMeshDescription == &RawMesh)
 					{
 						MeshData.CustomTextureCoordinates = UniqueTextureCoordinates;
 					}
@@ -2225,49 +2276,50 @@ void FMeshMergeUtilities::MergeComponentsToStaticMesh(const TArray<UPrimitiveCom
 					// Found component entry
 
 					// Retrieve raw mesh data for this component and lod pair
-					FRawMesh* RawMesh = DataTracker.GetRawMeshPtr(MappingPair.Key);
+					FMeshDescription* RawMesh = DataTracker.GetRawMeshPtr(MappingPair.Key);
 
 					FMeshData& MeshData = GlobalMeshSettings[MappingPair.Value.Value];
 					const FUVOffsetScalePair& UVTransform = UVTransforms[MappingPair.Value.Value];
 
 					const uint32 MaterialIndex = MappingPair.Value.Key;
 					ProcessedMaterials.Add(MaterialIndex);
-					if (RawMesh->VertexPositions.Num())
+					if (RawMesh->Vertices().Num())
 					{
-						for (int32 UVChannelIdx = 0; UVChannelIdx < MAX_MESH_TEXTURE_COORDS; ++UVChannelIdx)
+						TVertexInstanceAttributesRef<FVector2D> VertexInstanceUVs = RawMesh->VertexInstanceAttributes().GetAttributesRef<FVector2D>(MeshAttribute::VertexInstance::TextureCoordinate);
+						int32 NumUVChannel = FMath::Min(VertexInstanceUVs.GetNumIndices(), (int32)MAX_MESH_TEXTURE_COORDS);
+						for (int32 UVChannelIdx = 0; UVChannelIdx < NumUVChannel; ++UVChannelIdx)
 						{
-							TArray<FVector2D>* UVPtr = &RawMesh->WedgeTexCoords[UVChannelIdx];
-							if(UVChannelIdx == 0)
+							int32 VertexIndex = 0;
+							for (FVertexInstanceID VertexInstanceID : RawMesh->VertexInstances().GetElementIDs())
 							{
-								if (MeshData.CustomTextureCoordinates.Num())
+								FVector2D UV = VertexInstanceUVs.Get(VertexInstanceID, UVChannelIdx);
+								if (UVChannelIdx == 0)
 								{
-									UVPtr = &MeshData.CustomTextureCoordinates;
-								}
-								else if (MeshData.TextureCoordinateIndex != 0)
-								{
-									UVPtr = &RawMesh->WedgeTexCoords[MeshData.TextureCoordinateIndex];
+									if (MeshData.CustomTextureCoordinates.Num())
+									{
+										UV = MeshData.CustomTextureCoordinates[VertexIndex];
+									}
+									else if (MeshData.TextureCoordinateIndex != 0)
+									{
+										check(MeshData.TextureCoordinateIndex < NumUVChannel);
+										UV = VertexInstanceUVs.Get(VertexInstanceID, MeshData.TextureCoordinateIndex);
+									}
 								}
-							}
-
-							TArray<FVector2D>& UVs = *UVPtr;
 
-							if (RawMesh->WedgeTexCoords[UVChannelIdx].Num() > 0)
-							{
-								int32 UVIdx = 0;
-								for (const int32 FaceMaterialIndex : RawMesh->FaceMaterialIndices)
+								const TArray<FPolygonID>& Polygons = RawMesh->GetVertexInstanceConnectedPolygons(VertexInstanceID);
+								for (FPolygonID PolygonID : Polygons)
 								{
-									if (FaceMaterialIndex == MaterialIndex)
+									FPolygonGroupID PolygonGroupID = RawMesh->GetPolygonPolygonGroup(PolygonID);
+									if (PolygonGroupID.GetValue() == MaterialIndex)
 									{
 										if (UVTransform.Value != FVector2D::ZeroVector)
 										{
-											RawMesh->WedgeTexCoords[UVChannelIdx][UVIdx + 0] = UVs[UVIdx + 0] * UVTransform.Value + UVTransform.Key;
-											RawMesh->WedgeTexCoords[UVChannelIdx][UVIdx + 1] = UVs[UVIdx + 1] * UVTransform.Value + UVTransform.Key;
-											RawMesh->WedgeTexCoords[UVChannelIdx][UVIdx + 2] = UVs[UVIdx + 2] * UVTransform.Value + UVTransform.Key;
+											VertexInstanceUVs.Set(VertexInstanceID, UVChannelIdx, UV * UVTransform.Value + UVTransform.Key);
+											break;
 										}
 									}
-
-									UVIdx += 3;
 								}
+								VertexIndex++;
 							}
 						}
 					}
@@ -2277,12 +2329,14 @@ void FMeshMergeUtilities::MergeComponentsToStaticMesh(const TArray<UPrimitiveCom
 
 		for (TRawMeshIterator Iterator = DataTracker.GetRawMeshIterator(); Iterator; ++Iterator)
 		{
-			FRawMesh& RawMesh = Iterator.Value();
+			FMeshDescription& RawMesh = Iterator.Value();
 			// Reset material indexes
-			for (int32& FaceMaterialIndex : RawMesh.FaceMaterialIndices)
+			TMap<FPolygonGroupID, FPolygonGroupID> RemapPolygonGroups;
+			for (FPolygonGroupID PolygonGroupID : RawMesh.PolygonGroups().GetElementIDs())
 			{
-				FaceMaterialIndex = 0;
+				RemapPolygonGroups.Add(PolygonGroupID, FPolygonGroupID(0));
 			}
+			FMeshDescriptionOperations::RemapPolygonGroups(RawMesh, RemapPolygonGroups);
 		}
 
 		UMaterialInterface* MergedMaterial = CreateProxyMaterial(InBasePackageName, MergedAssetPackageName, InBaseMaterial, InOuter, InSettings, OutMaterial, OutAssetsToSync);
@@ -2300,7 +2354,7 @@ void FMeshMergeUtilities::MergeComponentsToStaticMesh(const TArray<UPrimitiveCom
 		}
 	}
 
-	TArray<FRawMesh> MergedRawMeshes;
+	TArray<FMeshDescription> MergedRawMeshes;
 	TMultiMap<FMeshLODKey, MaterialRemapPair> OutputMaterialsMap;
 	CreateMergedRawMeshes(DataTracker, InSettings, StaticMeshComponentsToMerge, UniqueMaterials, CollapsedMaterialMap, OutputMaterialsMap, bMergeAllLODs, bMergeMaterialData, MergedAssetPivot, MergedRawMeshes);
 
@@ -2310,12 +2364,24 @@ void FMeshMergeUtilities::MergeComponentsToStaticMesh(const TArray<UPrimitiveCom
 	{
 		const int32 LODIndex = *Iterator;
 		TArray<uint32> UniqueMaterialIndices;
-		const FRawMesh& TargetRawMesh = MergedRawMeshes[LODIndex];
-		for (uint32 MaterialIndex : TargetRawMesh.FaceMaterialIndices)
+		const FMeshDescription& TargetRawMesh = MergedRawMeshes[LODIndex];
+		uint32 MaterialIndex = 0;
+		for (FPolygonGroupID PolygonGroupID : TargetRawMesh.PolygonGroups().GetElementIDs())
 		{
-			UniqueMaterialIndices.AddUnique(MaterialIndex);
+			//Skip empty group
+			if (TargetRawMesh.GetPolygonGroupPolygons(PolygonGroupID).Num() > 0)
+			{
+				if (PolygonGroupID.GetValue() < DataTracker.NumberOfUniqueSections())
+				{
+					UniqueMaterialIndices.AddUnique(PolygonGroupID.GetValue());
+				}
+				else
+				{
+					UniqueMaterialIndices.AddUnique(MaterialIndex);
+				}
+				MaterialIndex++;
+			}
 		}
-
 		UniqueMaterialIndices.Sort();
 		for (int32 Index = 0; Index < UniqueMaterialIndices.Num(); ++Index)
 		{
@@ -2404,9 +2470,9 @@ void FMeshMergeUtilities::MergeComponentsToStaticMesh(const TArray<UPrimitiveCom
 		TArray<UMaterialInterface*> ImposterMaterials;
 		FBox ImposterBounds(EForceInit::ForceInit);
 		for (int32 LODIndex = 0; LODIndex < MergedRawMeshes.Num(); ++LODIndex)
-		{	
-			FRawMesh& MergedMeshLOD = MergedRawMeshes[LODIndex];
-			if (MergedMeshLOD.VertexPositions.Num() > 0)
+		{
+			FMeshDescription& MergedMeshLOD = MergedRawMeshes[LODIndex];
+			if (MergedMeshLOD.Vertices().Num() > 0)
 			{
 				FStaticMeshSourceModel& SrcModel = StaticMesh->AddSourceModel();
 				// Don't allow the engine to recalculate normals
@@ -2439,8 +2505,9 @@ void FMeshMergeUtilities::MergeComponentsToStaticMesh(const TArray<UPrimitiveCom
 						}
 					}
 				}
-
-				SrcModel.SaveRawMesh(MergedMeshLOD);
+				FMeshDescription* MeshDescription = StaticMesh->CreateOriginalMeshDescription(LODIndex);
+				*MeshDescription = MergedMeshLOD;
+				StaticMesh->CommitOriginalMeshDescription(LODIndex);
 			}
 		}
 		
@@ -2537,7 +2604,7 @@ void FMeshMergeUtilities::MergeComponentsToStaticMesh(const TArray<UPrimitiveCom
 	}
 }
 
-void FMeshMergeUtilities::CreateMergedRawMeshes(FMeshMergeDataTracker& InDataTracker, const FMeshMergingSettings& InSettings, const TArray<UStaticMeshComponent*>& InStaticMeshComponentsToMerge, const TArray<UMaterialInterface*>& InUniqueMaterials, const TMap<UMaterialInterface*, UMaterialInterface*>& InCollapsedMaterialMap, const TMultiMap<FMeshLODKey, MaterialRemapPair>& InOutputMaterialsMap, bool bInMergeAllLODs, bool bInMergeMaterialData, const FVector& InMergedAssetPivot, TArray<FRawMesh>& OutMergedRawMeshes) const
+void FMeshMergeUtilities::CreateMergedRawMeshes(FMeshMergeDataTracker& InDataTracker, const FMeshMergingSettings& InSettings, const TArray<UStaticMeshComponent*>& InStaticMeshComponentsToMerge, const TArray<UMaterialInterface*>& InUniqueMaterials, const TMap<UMaterialInterface*, UMaterialInterface*>& InCollapsedMaterialMap, const TMultiMap<FMeshLODKey, MaterialRemapPair>& InOutputMaterialsMap, bool bInMergeAllLODs, bool bInMergeMaterialData, const FVector& InMergedAssetPivot, TArray<FMeshDescription>& OutMergedRawMeshes) const
 {
 	if (bInMergeAllLODs)
 	{
@@ -2546,113 +2613,84 @@ void FMeshMergeUtilities::CreateMergedRawMeshes(FMeshMergeDataTracker& InDataTra
 		{
 			// Find meshes for each lod
 			const int32 LODIndex = *Iterator;
-			FRawMesh& MergedMesh = OutMergedRawMeshes[LODIndex];
+			FMeshDescription& MergedMesh = OutMergedRawMeshes[LODIndex];
+			UStaticMesh::RegisterMeshAttributes(MergedMesh);
+
 			for (int32 ComponentIndex = 0; ComponentIndex < InStaticMeshComponentsToMerge.Num(); ++ComponentIndex)
 			{
 				int32 RetrievedLODIndex = LODIndex;
-				FRawMesh* RawMeshPtr = InDataTracker.TryFindRawMeshForLOD(ComponentIndex, RetrievedLODIndex);
+				FMeshDescription* RawMeshPtr = InDataTracker.TryFindRawMeshForLOD(ComponentIndex, RetrievedLODIndex);
+
 
 				if (RawMeshPtr != nullptr)
 				{
-					MergedMesh.FaceSmoothingMasks.Append(RawMeshPtr->FaceSmoothingMasks);
+					InDataTracker.AddComponentToWedgeMapping(ComponentIndex, LODIndex, MergedMesh.VertexInstances().Num());
 
-					if (bInMergeMaterialData)
-					{
-						MergedMesh.FaceMaterialIndices.AddZeroed(RawMeshPtr->FaceMaterialIndices.Num());
-					}
-					else
+					FMeshDescriptionOperations::FAppendSettings AppendSettings;
+
+					AppendSettings.PolygonGroupsDelegate = FAppendPolygonGroupsDelegate::CreateLambda([&bInMergeMaterialData, &InDataTracker, &InOutputMaterialsMap, &ComponentIndex, &LODIndex](const FMeshDescription& SourceMesh, FMeshDescription& TargetMesh, PolygonGroupMap& RemapPolygonGroups)
 					{
-						TArray<SectionRemapPair> SectionMappings;
-						InDataTracker.GetMappingsForMeshLOD(FMeshLODKey(ComponentIndex, LODIndex), SectionMappings);
-	
-						for (const int32 FaceIndex : RawMeshPtr->FaceMaterialIndices)
+						TPolygonGroupAttributesConstRef<FName> SourceImportedMaterialSlotNames = SourceMesh.PolygonGroupAttributes().GetAttributesRef<FName>(MeshAttribute::PolygonGroup::ImportedMaterialSlotName);
+						TPolygonGroupAttributesRef<FName> TargetImportedMaterialSlotNames = TargetMesh.PolygonGroupAttributes().GetAttributesRef<FName>(MeshAttribute::PolygonGroup::ImportedMaterialSlotName);
+						//Copy the polygon group
+						if (bInMergeMaterialData)
 						{
-							// First map from original section index to unique material index
-							int32 UniqueIndex = INDEX_NONE;
-							// then map to the output material map, if any
-							if(InOutputMaterialsMap.Num() > 0)
+							FPolygonGroupID PolygonGroupID(0);
+							if (!TargetMesh.PolygonGroups().IsValid(PolygonGroupID))
 							{
-								TArray<MaterialRemapPair> MaterialMappings;
-								InOutputMaterialsMap.MultiFind(FMeshLODKey(ComponentIndex, LODIndex), MaterialMappings);
-								for(MaterialRemapPair& Pair : MaterialMappings)
-								{
-									if(Pair.Key == FaceIndex)
-									{
-										UniqueIndex = Pair.Value;
-										break;
-									}
-								}
-
-								// Note that at this point UniqueIndex is NOT a material index, but a unique section index!
+								TargetMesh.CreatePolygonGroupWithID(PolygonGroupID);
+								TargetImportedMaterialSlotNames[PolygonGroupID] = SourceMesh.PolygonGroups().IsValid(PolygonGroupID) ? SourceImportedMaterialSlotNames[PolygonGroupID] : FName(TEXT("DefaultMaterialName"));
 							}
-							else
+							for (FPolygonGroupID SourcePolygonGroupID : SourceMesh.PolygonGroups().GetElementIDs())
 							{
-								UniqueIndex = FaceIndex;
+								RemapPolygonGroups.Add(SourcePolygonGroupID, PolygonGroupID);
 							}
-
-							MergedMesh.FaceMaterialIndices.Add(UniqueIndex);
-						}
-					}
-
-					InDataTracker.AddComponentToWedgeMapping(ComponentIndex, LODIndex, MergedMesh.WedgeIndices.Num());
-					const int32 VertexOffset = MergedMesh.VertexPositions.Num();
-					for (int32 WedgeIndex : RawMeshPtr->WedgeIndices)
-					{
-						MergedMesh.WedgeIndices.Add(WedgeIndex + VertexOffset);
-					}
-
-					for (FVector VertexPos : RawMeshPtr->VertexPositions)
-					{
-						MergedMesh.VertexPositions.Add(VertexPos - InMergedAssetPivot);
-					}
-
-					MergedMesh.WedgeTangentX.Append(RawMeshPtr->WedgeTangentX);
-					MergedMesh.WedgeTangentY.Append(RawMeshPtr->WedgeTangentY);
-					MergedMesh.WedgeTangentZ.Append(RawMeshPtr->WedgeTangentZ);
-
-					// Deal with vertex colors
-					// Some meshes may have it, in this case merged mesh will be forced to have vertex colors as well
-					if (InSettings.bBakeVertexDataToMesh)
-					{
-						if (InDataTracker.DoesLODContainVertexColors(0) && RawMeshPtr->WedgeColors.Num())
-						{
-							MergedMesh.WedgeColors.Append(RawMeshPtr->WedgeColors);
 						}
 						else
 						{
-							// In case this source mesh does not have vertex colors, fill target with 0xFF
-							int32 ColorsOffset = MergedMesh.WedgeColors.Num();
-							int32 ColorsNum = RawMeshPtr->WedgeIndices.Num();
-							MergedMesh.WedgeColors.AddUninitialized(ColorsNum);
-							FMemory::Memset(&MergedMesh.WedgeColors[ColorsOffset], 0xFF, ColorsNum*MergedMesh.WedgeColors.GetTypeSize());
-						}
-					}
-
-					// Merge all other UV channels 
-					for (int32 ChannelIdx = 0; ChannelIdx < MAX_MESH_TEXTURE_COORDS; ++ChannelIdx)
-					{
-						// Whether this channel has data
-						if (InDataTracker.DoesUVChannelContainData(ChannelIdx, RetrievedLODIndex))
-						{
-							const TArray<FVector2D>& SourceChannel = RawMeshPtr->WedgeTexCoords[ChannelIdx];
-							TArray<FVector2D>& TargetChannel = MergedMesh.WedgeTexCoords[ChannelIdx];
-
-							// Whether source mesh has data in this channel
-							if (SourceChannel.Num())
-							{
-								TargetChannel.Append(SourceChannel);
-							}
-							else
+							TArray<SectionRemapPair> SectionMappings;
+							InDataTracker.GetMappingsForMeshLOD(FMeshLODKey(ComponentIndex, LODIndex), SectionMappings);
+							for (FPolygonGroupID SourcePolygonGroupID : SourceMesh.PolygonGroups().GetElementIDs())
 							{
-								// Fill with zero coordinates if source mesh has no data for this channel
-								const int32 TexCoordNum = RawMeshPtr->WedgeIndices.Num();
-								for (int32 CoordIdx = 0; CoordIdx < TexCoordNum; ++CoordIdx)
+								// First map from original section index to unique material index
+								int32 UniqueIndex = INDEX_NONE;
+								// then map to the output material map, if any
+								if (InOutputMaterialsMap.Num() > 0)
 								{
-									TargetChannel.Add(FVector2D::ZeroVector);
+									TArray<MaterialRemapPair> MaterialMappings;
+									InOutputMaterialsMap.MultiFind(FMeshLODKey(ComponentIndex, LODIndex), MaterialMappings);
+									for (MaterialRemapPair& Pair : MaterialMappings)
+									{
+										if (Pair.Key == SourcePolygonGroupID.GetValue())
+										{
+											UniqueIndex = Pair.Value;
+											break;
+										}
+									}
+
+									// Note that at this point UniqueIndex is NOT a material index, but a unique section index!
+								}
+								else
+								{
+									UniqueIndex = SourcePolygonGroupID.GetValue();
+								}
+								FPolygonGroupID TargetPolygonGroupID(UniqueIndex);
+								if (!TargetMesh.PolygonGroups().IsValid(TargetPolygonGroupID))
+								{
+									while (TargetMesh.PolygonGroups().Num() <= UniqueIndex)
+									{
+										TargetPolygonGroupID = TargetMesh.CreatePolygonGroup();
+									}
+									check(TargetPolygonGroupID.GetValue() == UniqueIndex);
+									TargetImportedMaterialSlotNames[TargetPolygonGroupID] = SourceImportedMaterialSlotNames[SourcePolygonGroupID];
 								}
+								RemapPolygonGroups.Add(SourcePolygonGroupID, TargetPolygonGroupID);
 							}
 						}
-					}
+					});
+					AppendSettings.bMergeVertexColor = InSettings.bBakeVertexDataToMesh;
+					AppendSettings.MergedAssetPivot = InMergedAssetPivot;
+					FMeshDescriptionOperations::AppendMeshDescription(*RawMeshPtr, MergedMesh, AppendSettings);
 				}
 			}
 		}
@@ -2660,120 +2698,87 @@ void FMeshMergeUtilities::CreateMergedRawMeshes(FMeshMergeDataTracker& InDataTra
 	else
 	{	
 		OutMergedRawMeshes.AddZeroed(1);
-		FRawMesh& MergedMesh = OutMergedRawMeshes.Last();
+		FMeshDescription& MergedMesh = OutMergedRawMeshes.Last();
+		UStaticMesh::RegisterMeshAttributes(MergedMesh);
+
 		for (int32 ComponentIndex = 0; ComponentIndex < InStaticMeshComponentsToMerge.Num(); ++ComponentIndex)
 		{
 			int32 LODIndex = 0;
 		
-			FRawMesh* RawMeshPtr = InDataTracker.FindRawMeshAndLODIndex(ComponentIndex, LODIndex);
+			FMeshDescription* RawMeshPtr = InDataTracker.FindRawMeshAndLODIndex(ComponentIndex, LODIndex);
 
 			if (RawMeshPtr != nullptr)
 			{
-				FRawMesh& RawMesh = *RawMeshPtr;
+				FMeshDescription& RawMesh = *RawMeshPtr;
 
-				MergedMesh.FaceSmoothingMasks.Append(RawMesh.FaceSmoothingMasks);
+				const int32 TargetLODIndex = 0;
+				InDataTracker.AddComponentToWedgeMapping(ComponentIndex, TargetLODIndex, MergedMesh.VertexInstances().Num());
 
-				if (bInMergeMaterialData)
-				{
-					MergedMesh.FaceMaterialIndices.AddZeroed(RawMesh.FaceMaterialIndices.Num());
-				}
-				else
-				{
-					TArray<SectionRemapPair> SectionMappings;
-					InDataTracker.GetMappingsForMeshLOD(FMeshLODKey(ComponentIndex, LODIndex), SectionMappings);
+				FMeshDescriptionOperations::FAppendSettings AppendSettings;
 
-					for (const int32 FaceIndex : RawMesh.FaceMaterialIndices)
+				AppendSettings.PolygonGroupsDelegate = FAppendPolygonGroupsDelegate::CreateLambda([&bInMergeMaterialData, &InDataTracker, &InOutputMaterialsMap, &ComponentIndex, &LODIndex](const FMeshDescription& SourceMesh, FMeshDescription& TargetMesh, PolygonGroupMap& RemapPolygonGroups)
+				{
+					TPolygonGroupAttributesConstRef<FName> SourceImportedMaterialSlotNames = SourceMesh.PolygonGroupAttributes().GetAttributesRef<FName>(MeshAttribute::PolygonGroup::ImportedMaterialSlotName);
+					TPolygonGroupAttributesRef<FName> TargetImportedMaterialSlotNames = TargetMesh.PolygonGroupAttributes().GetAttributesRef<FName>(MeshAttribute::PolygonGroup::ImportedMaterialSlotName);
+					//Copy the polygon group
+					if (bInMergeMaterialData)
 					{
-						// First map from original section index to unique material index
-						int32 UniqueIndex = INDEX_NONE;
-
-						// then map to the output material map, if any
-						if(InOutputMaterialsMap.Num() > 0)
+						FPolygonGroupID PolygonGroupID(0);
+						if (!TargetMesh.PolygonGroups().IsValid(PolygonGroupID))
 						{
-							TArray<MaterialRemapPair> MaterialMappings;
-							InOutputMaterialsMap.MultiFind(FMeshLODKey(ComponentIndex, LODIndex), MaterialMappings);
-							for(MaterialRemapPair& Pair : MaterialMappings)
-							{
-								if(Pair.Key == FaceIndex)
-								{
-									UniqueIndex = Pair.Value;
-									break;
-								}
-							}
-
-							// Note that at this point UniqueIndex is NOT a material index, but a unique section index!
+							TargetMesh.CreatePolygonGroupWithID(PolygonGroupID);
+							TargetImportedMaterialSlotNames[PolygonGroupID] = SourceMesh.PolygonGroups().IsValid(PolygonGroupID) ? SourceImportedMaterialSlotNames[PolygonGroupID] : FName(TEXT("DefaultMaterialName"));
 						}
-						else
+						for (FPolygonGroupID SourcePolygonGroupID : SourceMesh.PolygonGroups().GetElementIDs())
 						{
-							UniqueIndex = FaceIndex;
+							RemapPolygonGroups.Add(SourcePolygonGroupID, PolygonGroupID);
 						}
-
-						MergedMesh.FaceMaterialIndices.Add(UniqueIndex);
-					}
-				}
-
-				const int32 TargetLODIndex = 0;
-				InDataTracker.AddComponentToWedgeMapping(ComponentIndex, TargetLODIndex, MergedMesh.WedgeIndices.Num());
-
-				const int32 VertexOffset = MergedMesh.VertexPositions.Num();
-
-				for (int32 Index : RawMesh.WedgeIndices)
-				{
-					MergedMesh.WedgeIndices.Add(Index + VertexOffset);
-				}
-
-				for (FVector VertexPos : RawMesh.VertexPositions)
-				{
-					MergedMesh.VertexPositions.Add(VertexPos - InMergedAssetPivot);
-				}
-
-				MergedMesh.WedgeTangentX.Append(RawMesh.WedgeTangentX);
-				MergedMesh.WedgeTangentY.Append(RawMesh.WedgeTangentY);
-				MergedMesh.WedgeTangentZ.Append(RawMesh.WedgeTangentZ);
-
-				// Deal with vertex colors
-				// Some meshes may have it, in this case merged mesh will be forced to have vertex colors as well
-				if (InSettings.bBakeVertexDataToMesh)
-				{
-					if (InDataTracker.DoesAnyLODContainVertexColors() && RawMesh.WedgeColors.Num())
-					{
-						MergedMesh.WedgeColors.Append(RawMesh.WedgeColors);
 					}
 					else
 					{
-						// In case this source mesh does not have vertex colors, fill target with 0xFF
-						int32 ColorsOffset = MergedMesh.WedgeColors.Num();
-						int32 ColorsNum = RawMesh.WedgeIndices.Num();
-						MergedMesh.WedgeColors.AddUninitialized(ColorsNum);
-						FMemory::Memset(&MergedMesh.WedgeColors[ColorsOffset], 0xFF, ColorsNum*MergedMesh.WedgeColors.GetTypeSize());
-					}
-				}
-
-				// Merge all other UV channels 
-				for (int32 ChannelIdx = 0; ChannelIdx < MAX_MESH_TEXTURE_COORDS; ++ChannelIdx)
-				{
-					// Whether this channel has data
-					if (InDataTracker.DoesUVChannelContainData(ChannelIdx) && InSettings.OutputUVs[ChannelIdx] == EUVOutput::OutputChannel)
-					{
-						const TArray<FVector2D>& SourceChannel = RawMesh.WedgeTexCoords[ChannelIdx];
-						TArray<FVector2D>& TargetChannel = MergedMesh.WedgeTexCoords[ChannelIdx];
-
-						// Whether source mesh has data in this channel
-						if (SourceChannel.Num())
-						{
-							TargetChannel.Append(SourceChannel);
-						}
-						else
+						TArray<SectionRemapPair> SectionMappings;
+						InDataTracker.GetMappingsForMeshLOD(FMeshLODKey(ComponentIndex, LODIndex), SectionMappings);
+						for (FPolygonGroupID SourcePolygonGroupID : SourceMesh.PolygonGroups().GetElementIDs())
 						{
-							// Fill with zero coordinates if source mesh has no data for this channel
-							const int32 TexCoordNum = RawMesh.WedgeIndices.Num();
-							for (int32 CoordIdx = 0; CoordIdx < TexCoordNum; ++CoordIdx)
+							// First map from original section index to unique material index
+							int32 UniqueIndex = INDEX_NONE;
+							// then map to the output material map, if any
+							if (InOutputMaterialsMap.Num() > 0)
 							{
-								TargetChannel.Add(FVector2D::ZeroVector);
+								TArray<MaterialRemapPair> MaterialMappings;
+								InOutputMaterialsMap.MultiFind(FMeshLODKey(ComponentIndex, LODIndex), MaterialMappings);
+								for (MaterialRemapPair& Pair : MaterialMappings)
+								{
+									if (Pair.Key == SourcePolygonGroupID.GetValue())
+									{
+										UniqueIndex = Pair.Value;
+										break;
+									}
+								}
+
+								// Note that at this point UniqueIndex is NOT a material index, but a unique section index!
+							}
+							else
+							{
+								UniqueIndex = SourcePolygonGroupID.GetValue();
 							}
+							FPolygonGroupID TargetPolygonGroupID(UniqueIndex);
+							if (!TargetMesh.PolygonGroups().IsValid(TargetPolygonGroupID))
+							{
+								while (TargetMesh.PolygonGroups().Num() <= UniqueIndex)
+								{
+									TargetPolygonGroupID = TargetMesh.CreatePolygonGroup();
+								}
+								check(TargetPolygonGroupID.GetValue() == UniqueIndex);
+								TargetImportedMaterialSlotNames[TargetPolygonGroupID] = SourceImportedMaterialSlotNames[SourcePolygonGroupID];
+							}
+							RemapPolygonGroups.Add(SourcePolygonGroupID, TargetPolygonGroupID);
 						}
 					}
-				}
+				});
+				AppendSettings.bMergeVertexColor = InSettings.bBakeVertexDataToMesh;
+				AppendSettings.MergedAssetPivot = InMergedAssetPivot;
+				FMeshDescriptionOperations::AppendMeshDescription(*RawMeshPtr, MergedMesh, AppendSettings);
 			}
 		}
 	}
diff --git a/Engine/Source/Developer/MeshMergeUtilities/Private/MeshMergeUtilities.h b/Engine/Source/Developer/MeshMergeUtilities/Private/MeshMergeUtilities.h
--- a/Engine/Source/Developer/MeshMergeUtilities/Private/MeshMergeUtilities.h
+++ b/Engine/Source/Developer/MeshMergeUtilities/Private/MeshMergeUtilities.h
@@ -53,7 +53,7 @@ protected:
 	/** Merges flattened material into binned textures */
 	void FlattenBinnedMaterials(TArray<struct FFlattenMaterial>& InMaterialList, const TArray<FBox2D>& InMaterialBoxes, int32 InGutter, bool bCopyOnlyMaskedPixels, FFlattenMaterial& OutMergedMaterial, TArray<FUVOffsetScalePair>& OutUVTransforms) const;
 	/** Helper function to create the final merged raw meshes */
-	void CreateMergedRawMeshes(FMeshMergeDataTracker& InDataTracker, const FMeshMergingSettings& InSettings, const TArray<UStaticMeshComponent*>& InStaticMeshComponentsToMerge, const TArray<UMaterialInterface*>& InUniqueMaterials, const TMap<UMaterialInterface*, UMaterialInterface*>& InCollapsedMaterialMap, const TMultiMap<FMeshLODKey, MaterialRemapPair>& InOutputMaterialsMap, bool bInMergeAllLODs, bool bInMergeMaterialData, const FVector& InMergedAssetPivot, TArray<FRawMesh>& OutMergedRawMeshes) const;
+	void CreateMergedRawMeshes(FMeshMergeDataTracker& InDataTracker, const FMeshMergingSettings& InSettings, const TArray<UStaticMeshComponent*>& InStaticMeshComponentsToMerge, const TArray<UMaterialInterface*>& InUniqueMaterials, const TMap<UMaterialInterface*, UMaterialInterface*>& InCollapsedMaterialMap, const TMultiMap<FMeshLODKey, MaterialRemapPair>& InOutputMaterialsMap, bool bInMergeAllLODs, bool bInMergeMaterialData, const FVector& InMergedAssetPivot, TArray<FMeshDescription>& OutMergedRawMeshes) const;
 protected:
 	/** Flattens out emissive scale across all flatten material instances */
 	float FlattenEmissivescale(TArray<struct FFlattenMaterial>& InMaterialList) const;
diff --git a/Engine/Source/Developer/MeshMergeUtilities/Private/ProxyGenerationProcessor.cpp b/Engine/Source/Developer/MeshMergeUtilities/Private/ProxyGenerationProcessor.cpp
--- a/Engine/Source/Developer/MeshMergeUtilities/Private/ProxyGenerationProcessor.cpp
+++ b/Engine/Source/Developer/MeshMergeUtilities/Private/ProxyGenerationProcessor.cpp
@@ -9,6 +9,10 @@
 #include "IMeshReductionManagerModule.h"
 #include "Modules/ModuleManager.h"
 
+#include "MeshDescription.h"
+#include "MeshAttributes.h"
+#include "MeshAttributeArray.h"
+
 #if WITH_EDITOR
 #include "Editor.h"
 #include "MeshMergeHelpers.h"
@@ -82,7 +86,7 @@ bool FProxyGenerationProcessor::Tick(float DeltaTime)
 	return true;
 }
 
-void FProxyGenerationProcessor::ProxyGenerationComplete(FRawMesh& OutProxyMesh, struct FFlattenMaterial& OutMaterial, const FGuid OutJobGUID)
+void FProxyGenerationProcessor::ProxyGenerationComplete(FMeshDescription& OutProxyMesh, struct FFlattenMaterial& OutMaterial, const FGuid OutJobGUID)
 {
 	FScopeLock Lock(&StateLock);
 	FMergeCompleteData** FindData = ProxyMeshJobs.Find(OutJobGUID);
@@ -200,6 +204,26 @@ void FProxyGenerationProcessor::ProcessJob(const FGuid& JobGuid, FProxyGeneratio
 
 	const bool bContainsImposters = Data->MergeData->ImposterComponents.Num() > 0;
 	FBox ImposterBounds(EForceInit::ForceInit);
+
+	auto RemoveVertexColorAndCommitMeshDescription = [&StaticMesh, &Data]()
+	{
+		if (!Data->MergeData->InProxySettings.bAllowVertexColors)
+		{
+			//We cannot remove the vertex color with the mesh description so we assign a white value to all color
+			TVertexInstanceAttributesRef<FVector4> VertexInstanceColors = Data->RawMesh.VertexInstanceAttributes().GetAttributesRef<FVector4>(MeshAttribute::VertexInstance::Color);
+			//set all value to white
+			for (const FVertexInstanceID& VertexInstanceID : Data->RawMesh.VertexInstances().GetElementIDs())
+			{
+				VertexInstanceColors[VertexInstanceID] = FVector4(1.0f, 1.0f, 1.0f);
+			}
+		}
+		//Commit the FMeshDescription to the staticmesh
+		int32 LodIndex = StaticMesh->GetNumLODs() - 1;
+		FMeshDescription* OriginalMeshDescription = StaticMesh->GetOriginalMeshDescription(LodIndex);
+		*OriginalMeshDescription = Data->RawMesh;
+		StaticMesh->CommitOriginalMeshDescription(LodIndex);
+	};
+
 	if (bContainsImposters)
 	{
 		TArray<UMaterialInterface*> ImposterMaterials;
@@ -216,13 +240,8 @@ void FProxyGenerationProcessor::ProcessJob(const FGuid& JobGuid, FProxyGeneratio
 				ImposterBounds += Component->GetStaticMesh()->GetBoundingBox().TransformBy(Component->GetComponentToWorld());
 			}
 		}
+		RemoveVertexColorAndCommitMeshDescription();
 
-		if (!Data->MergeData->InProxySettings.bAllowVertexColors)
-		{
-			Data->RawMesh.WedgeColors.Empty();
-		}
-
-		SrcModel.SaveRawMesh(Data->RawMesh);
 		StaticMesh->StaticMaterials.Add(FStaticMaterial(ProxyMaterial));
 
 		for (UMaterialInterface* Material : ImposterMaterials)
@@ -232,13 +251,8 @@ void FProxyGenerationProcessor::ProcessJob(const FGuid& JobGuid, FProxyGeneratio
 	}
 	else
 	{
-		if (!Data->MergeData->InProxySettings.bAllowVertexColors)
-		{
-			Data->RawMesh.WedgeColors.Empty();
-		}
-
+		RemoveVertexColorAndCommitMeshDescription();
 
-		SrcModel.SaveRawMesh(Data->RawMesh);
 		StaticMesh->StaticMaterials.Add(FStaticMaterial(ProxyMaterial));
 	}	
 
@@ -246,10 +260,28 @@ void FProxyGenerationProcessor::ProcessJob(const FGuid& JobGuid, FProxyGeneratio
 	StaticMesh->ImportVersion = EImportStaticMeshVersion::LastVersion;
 
 	// setup section info map
+	TPolygonGroupAttributesConstRef<FName> PolygonGroupImportedMaterialSlotNames = Data->RawMesh.PolygonGroupAttributes().GetAttributesRef<FName>(MeshAttribute::PolygonGroup::ImportedMaterialSlotName);
 	TArray<int32> UniqueMaterialIndices;
-	for (int32 MaterialIndex : Data->RawMesh.FaceMaterialIndices)
+	for (const FPolygonGroupID& PolygonGroupID : Data->RawMesh.PolygonGroups().GetElementIDs())
 	{
-		UniqueMaterialIndices.AddUnique(MaterialIndex);
+		int32 PolygonGroupMaterialIndex = PolygonGroupID.GetValue();
+		FName PolygonGroupName = PolygonGroupImportedMaterialSlotNames[PolygonGroupID];
+		if (PolygonGroupName != NAME_None)
+		{
+			for (int32 MaterialIndex = 0; MaterialIndex < StaticMesh->StaticMaterials.Num(); ++MaterialIndex)
+			{
+				if (StaticMesh->StaticMaterials[MaterialIndex].ImportedMaterialSlotName == PolygonGroupName)
+				{
+					PolygonGroupMaterialIndex = MaterialIndex;
+					break;
+				}
+			}
+		}
+		if (!StaticMesh->StaticMaterials.IsValidIndex(PolygonGroupMaterialIndex))
+		{
+			PolygonGroupMaterialIndex = 0;
+		}
+		UniqueMaterialIndices.AddUnique(PolygonGroupMaterialIndex);
 	}
 
 	int32 SectionIndex = 0;
diff --git a/Engine/Source/Developer/MeshMergeUtilities/Private/ProxyGenerationProcessor.h b/Engine/Source/Developer/MeshMergeUtilities/Private/ProxyGenerationProcessor.h
--- a/Engine/Source/Developer/MeshMergeUtilities/Private/ProxyGenerationProcessor.h
+++ b/Engine/Source/Developer/MeshMergeUtilities/Private/ProxyGenerationProcessor.h
@@ -3,7 +3,7 @@
 #pragma once
 
 #include "Containers/Ticker.h"
-#include "RawMesh.h"
+#include "MeshDescription.h"
 #include "MaterialUtilities.h"
 
 struct FMergeCompleteData;
@@ -20,7 +20,7 @@ public:
 
 	void AddProxyJob(FGuid InJobGuid, FMergeCompleteData* InCompleteData);
 	virtual bool Tick(float DeltaTime) override;
-	void ProxyGenerationComplete(FRawMesh& OutProxyMesh, struct FFlattenMaterial& OutMaterial, const FGuid OutJobGUID);
+	void ProxyGenerationComplete(FMeshDescription& OutProxyMesh, struct FFlattenMaterial& OutMaterial, const FGuid OutJobGUID);
 
 	//@third party BEGIN SIMPLYGON
 	void ProxyGenerationFailed(const FGuid OutJobGUID, const FString& ErrorMessage);
@@ -40,7 +40,7 @@ protected:
 	/** Structure storing the data required during processing */
 	struct FProxyGenerationData
 	{
-		FRawMesh RawMesh;
+		FMeshDescription RawMesh;
 		FFlattenMaterial Material;
 		FMergeCompleteData* MergeData;
 	};
diff --git a/Engine/Source/Developer/MeshMergeUtilities/Private/SkeletalMeshAdapter.cpp b/Engine/Source/Developer/MeshMergeUtilities/Private/SkeletalMeshAdapter.cpp
--- a/Engine/Source/Developer/MeshMergeUtilities/Private/SkeletalMeshAdapter.cpp
+++ b/Engine/Source/Developer/MeshMergeUtilities/Private/SkeletalMeshAdapter.cpp
@@ -19,7 +19,7 @@ int32 FSkeletalMeshComponentAdapter::GetNumberOfLODs() const
 	return NumLODs;
 }
 
-void FSkeletalMeshComponentAdapter::RetrieveRawMeshData(int32 LODIndex, FRawMesh& InOutRawMesh, bool bPropogateMeshData) const
+void FSkeletalMeshComponentAdapter::RetrieveRawMeshData(int32 LODIndex, FMeshDescription& InOutRawMesh, bool bPropogateMeshData) const
 {
 	FMeshMergeHelpers::RetrieveMesh(SkeletalMeshComponent, LODIndex, InOutRawMesh, bPropogateMeshData);
 }
diff --git a/Engine/Source/Developer/MeshMergeUtilities/Private/StaticMeshAdapter.cpp b/Engine/Source/Developer/MeshMergeUtilities/Private/StaticMeshAdapter.cpp
--- a/Engine/Source/Developer/MeshMergeUtilities/Private/StaticMeshAdapter.cpp
+++ b/Engine/Source/Developer/MeshMergeUtilities/Private/StaticMeshAdapter.cpp
@@ -18,7 +18,7 @@ int32 FStaticMeshAdapter::GetNumberOfLODs() const
 	return NumLODs;
 }
 
-void FStaticMeshAdapter::RetrieveRawMeshData(int32 LODIndex, FRawMesh& InOutRawMesh, bool bPropogateMeshData) const
+void FStaticMeshAdapter::RetrieveRawMeshData(int32 LODIndex, FMeshDescription& InOutRawMesh, bool bPropogateMeshData) const
 {
 	FMeshMergeHelpers::RetrieveMesh(StaticMesh, LODIndex, InOutRawMesh);
 }
diff --git a/Engine/Source/Developer/MeshMergeUtilities/Private/StaticMeshComponentAdapter.cpp b/Engine/Source/Developer/MeshMergeUtilities/Private/StaticMeshComponentAdapter.cpp
--- a/Engine/Source/Developer/MeshMergeUtilities/Private/StaticMeshComponentAdapter.cpp
+++ b/Engine/Source/Developer/MeshMergeUtilities/Private/StaticMeshComponentAdapter.cpp
@@ -22,7 +22,7 @@ int32 FStaticMeshComponentAdapter::GetNumberOfLODs() const
 	return NumLODs;
 }
 
-void FStaticMeshComponentAdapter::RetrieveRawMeshData(int32 LODIndex, FRawMesh& InOutRawMesh, bool bPropogateMeshData) const
+void FStaticMeshComponentAdapter::RetrieveRawMeshData(int32 LODIndex, FMeshDescription& InOutRawMesh, bool bPropogateMeshData) const
 {
 	FMeshMergeHelpers::RetrieveMesh(StaticMeshComponent, LODIndex, InOutRawMesh, bPropogateMeshData);
 }
diff --git a/Engine/Source/Developer/MeshMergeUtilities/Public/IMaterialBakingAdapter.h b/Engine/Source/Developer/MeshMergeUtilities/Public/IMaterialBakingAdapter.h
--- a/Engine/Source/Developer/MeshMergeUtilities/Public/IMaterialBakingAdapter.h
+++ b/Engine/Source/Developer/MeshMergeUtilities/Public/IMaterialBakingAdapter.h
@@ -4,7 +4,7 @@
 
 #include "Math/BoxSphereBounds.h"
 
-struct FRawMesh;
+struct FMeshDescription;
 struct FSectionInfo;
 struct FMeshData;
 class UMaterialInterface;
@@ -14,8 +14,8 @@ class IMaterialBakingAdapter
 public:
 	/** Returns the number of LODs for the data the adapter represents */
 	virtual int32 GetNumberOfLODs() const = 0;
-	/** Retrieves model data in FRawMesh form */
-	virtual void RetrieveRawMeshData(int32 LODIndex, FRawMesh& InOutRawMesh, bool bPropagateMeshData) const = 0;
+	/** Retrieves model data in FMeshDescription form */
+	virtual void RetrieveRawMeshData(int32 LODIndex, FMeshDescription& InOutRawMesh, bool bPropagateMeshData) const = 0;
 	/** Retrieves all mesh sections from underlying data */
 	virtual void RetrieveMeshSections(int32 LODIndex, TArray<FSectionInfo>& InOutSectionInfo) const = 0;
 	/** Returns lightmap UV index used by mesh data */
diff --git a/Engine/Source/Developer/MeshMergeUtilities/Public/IMeshMergeExtension.h b/Engine/Source/Developer/MeshMergeUtilities/Public/IMeshMergeExtension.h
--- a/Engine/Source/Developer/MeshMergeUtilities/Public/IMeshMergeExtension.h
+++ b/Engine/Source/Developer/MeshMergeUtilities/Public/IMeshMergeExtension.h
@@ -4,13 +4,13 @@
 
 struct FRawMesh;
 struct FSectionInfo;
-struct FMeshData;
+struct FMeshDescription;
 class UMaterialInterface;
 
 class IMeshMergeExtension
 {
 public:
-	virtual void OnCreatedMergedRawMeshes(const TArray<UStaticMeshComponent*>& MergedComponents, const class FMeshMergeDataTracker& DataTracker, TArray<FRawMesh>& MergedMeshLODs) = 0;
+	virtual void OnCreatedMergedRawMeshes(const TArray<UStaticMeshComponent*>& MergedComponents, const class FMeshMergeDataTracker& DataTracker, TArray<FMeshDescription>& MergedMeshLODs) = 0;
 
 	virtual void OnCreatedProxyMaterial(const TArray<UStaticMeshComponent*>& MergedComponents, UMaterialInterface* ProxyMaterial) = 0;
 };
\ No newline at end of file
diff --git a/Engine/Source/Developer/MeshMergeUtilities/Public/MeshMergeData.h b/Engine/Source/Developer/MeshMergeUtilities/Public/MeshMergeData.h
--- a/Engine/Source/Developer/MeshMergeUtilities/Public/MeshMergeData.h
+++ b/Engine/Source/Developer/MeshMergeUtilities/Public/MeshMergeData.h
@@ -4,7 +4,6 @@
 
 #include "CoreMinimal.h"
 #include "Engine/StaticMesh.h"
-#include "RawMesh.h"
 #include "MeshDescription.h"
 #include "PhysicsEngine/AggregateGeom.h"
 
@@ -17,7 +16,6 @@ struct FMeshMergeData
 {
 	FMeshMergeData()
 		: RawMesh(nullptr)
-		, MeshDescription(nullptr)
 		, SourceStaticMesh(nullptr)
 		, bIsClippingMesh(false)
 	{}
@@ -29,19 +27,10 @@ struct FMeshMergeData
 			delete RawMesh;
 			RawMesh = nullptr;
 		}
-		if (MeshDescription != nullptr)
-		{
-			delete MeshDescription;
-			MeshDescription = nullptr;
-		}
 	}
 	
 	/** Raw mesh data from the source static mesh */
-	FRawMesh* RawMesh;
-	/** MeshDesciption pointer that get populated when we query it */
-	FMeshDescription* MeshDescription;
-	//MeshDescription getter
-	MESHMERGEUTILITIES_API FMeshDescription* GetMeshDescription();
+	FMeshDescription* RawMesh;	
 
 	/** Contains the original texture bounds, if the material requires baking down per-vertex data */
 	TArray<FBox2D> TexCoordBounds;
diff --git a/Engine/Source/Developer/MeshMergeUtilities/Public/MeshMergeDataTracker.h b/Engine/Source/Developer/MeshMergeUtilities/Public/MeshMergeDataTracker.h
--- a/Engine/Source/Developer/MeshMergeUtilities/Public/MeshMergeDataTracker.h
+++ b/Engine/Source/Developer/MeshMergeUtilities/Public/MeshMergeDataTracker.h
@@ -2,9 +2,12 @@
 
 #pragma once
 
-#include "RawMesh.h"
 #include "StaticMeshResources.h"
 #include "Materials/MaterialInstance.h"
+#include "MeshDescription.h"
+#include "MeshAttributes.h"
+#include "MeshAttributeArray.h"
+#include "MeshDescriptionOperations.h"
 
 typedef TPair<uint32, uint32> MeshLODPair;
 typedef TPair<uint32, uint32> SectionRemapPair;
@@ -103,8 +106,8 @@ struct MESHMERGEUTILITIES_API FMaterialKey
 };
 
 /** Typedefs to allow for some nicer looking loops */
-typedef TMap<FMeshLODKey, FRawMesh>::TConstIterator TConstRawMeshIterator;
-typedef TMap<FMeshLODKey, FRawMesh>::TIterator TRawMeshIterator;
+typedef TMap<FMeshLODKey, FMeshDescription>::TConstIterator TConstRawMeshIterator;
+typedef TMap<FMeshLODKey, FMeshDescription>::TIterator TRawMeshIterator;
 
 typedef TArray<int32>::TConstIterator TConstLODIndexIterator;
 
@@ -125,21 +128,21 @@ public:
 	void GetMappingsForMeshLOD(FMeshLODKey Key, TArray<SectionRemapPair>& InOutMappings);
 	
 	/** Adds or retrieves raw mesh data for the mesh and LOD index */
-	FRawMesh& AddAndRetrieveRawMesh(int32 MeshIndex, int32 LODIndex, UStaticMesh* InMesh);
+	FMeshDescription& AddAndRetrieveRawMesh(int32 MeshIndex, int32 LODIndex, UStaticMesh* InMesh);
 	/** Removes raw mesh entry for the given mesh and LOD index */
 	void RemoveRawMesh(int32 MeshIndex, int32 LODIndex);
 	/** Retrieves Raw Mesh ptr for the given mesh and LOD index */
-	FRawMesh* GetRawMeshPtr(int32 MeshIndex, int32 LODIndex);
+	FMeshDescription* GetRawMeshPtr(int32 MeshIndex, int32 LODIndex);
 	/** Retrieves Raw Mesh ptr for the given MeshLOD key */
-	FRawMesh* GetRawMeshPtr(FMeshLODKey Key);
-	/** Tries to retrieve a FRawMesh and returns the LOD index it found an entry for */
-	FRawMesh* FindRawMeshAndLODIndex(int32 MeshIndex, int32& OutLODIndex);
-	/** Tries to retrieve a FRawMesh for the given mesh and LOD index, if it can't it will try to find an entry for each LOD levle below InOutDesiredLODIndex */
-	FRawMesh* TryFindRawMeshForLOD(int32 MeshIndex, int32& InOutDesiredLODIndex);
+	FMeshDescription* GetRawMeshPtr(FMeshLODKey Key);
+	/** Tries to retrieve a FMeshDescription and returns the LOD index it found an entry for */
+	FMeshDescription* FindRawMeshAndLODIndex(int32 MeshIndex, int32& OutLODIndex);
+	/** Tries to retrieve a FMeshDescription for the given mesh and LOD index, if it can't it will try to find an entry for each LOD levle below InOutDesiredLODIndex */
+	FMeshDescription* TryFindRawMeshForLOD(int32 MeshIndex, int32& InOutDesiredLODIndex);
 
-	/** Returns a const key/value iterator for the FRawMesh entries */
+	/** Returns a const key/value iterator for the FMeshDescription entries */
 	TConstRawMeshIterator GetConstRawMeshIterator() const;
-	/** Returns a non-const key/value iterator for the FRawMesh entries */ 
+	/** Returns a non-const key/value iterator for the FMeshDescription entries */ 
 	TRawMeshIterator GetRawMeshIterator();
 
 	/** Adds a record of what channel lightmap data is stored at */
@@ -191,7 +194,7 @@ public:
 
 protected:
 	// Mesh / LOD index, RawMesh
-	TMap<FMeshLODKey, FRawMesh> RawMeshLODs;
+	TMap<FMeshLODKey, FMeshDescription> RawMeshLODs;
 
 	// Mesh / LOD index, lightmap channel
 	TMap<FMeshLODKey, int32> LightmapChannelLODs;
@@ -204,7 +207,7 @@ protected:
 	
 	/** Flags for UV and vertex color usage */
 	bool bWithVertexColors[MAX_STATIC_MESH_LODS];
-	bool bOcuppiedUVChannels[MAX_STATIC_MESH_LODS][MAX_MESH_TEXTURE_COORDS];
+	bool bOcuppiedUVChannels[MAX_STATIC_MESH_LODS][MAX_MESH_TEXTURE_COORDS_MD];
 	/** First available UV channel across all RawMesh entries */
 	int32 AvailableLightMapUVChannel;
 	int32 SummedLightMapPixels;
diff --git a/Engine/Source/Developer/MeshMergeUtilities/Public/SkeletalMeshAdapter.h b/Engine/Source/Developer/MeshMergeUtilities/Public/SkeletalMeshAdapter.h
--- a/Engine/Source/Developer/MeshMergeUtilities/Public/SkeletalMeshAdapter.h
+++ b/Engine/Source/Developer/MeshMergeUtilities/Public/SkeletalMeshAdapter.h
@@ -15,7 +15,7 @@ public:
 
 	/** Begin IMaterialBakingAdapter overrides */
 	virtual int32 GetNumberOfLODs() const override;
-	virtual void RetrieveRawMeshData(int32 LODIndex, FRawMesh& InOutRawMesh, bool bPropogateMeshData) const override;
+	virtual void RetrieveRawMeshData(int32 LODIndex, FMeshDescription& InOutRawMesh, bool bPropogateMeshData) const override;
 	virtual void RetrieveMeshSections(int32 LODIndex, TArray<FSectionInfo>& InOutSectionInfo) const override;
 	virtual int32 GetMaterialIndex(int32 LODIndex, int32 SectionIndex) const override;
 	virtual void ApplySettings(int32 LODIndex, FMeshData& InOutMeshData) const override {}
diff --git a/Engine/Source/Developer/MeshMergeUtilities/Public/StaticMeshAdapter.h b/Engine/Source/Developer/MeshMergeUtilities/Public/StaticMeshAdapter.h
--- a/Engine/Source/Developer/MeshMergeUtilities/Public/StaticMeshAdapter.h
+++ b/Engine/Source/Developer/MeshMergeUtilities/Public/StaticMeshAdapter.h
@@ -14,7 +14,7 @@ public:
 
 	/** Begin IMaterialBakingAdapter overrides */
 	virtual int32 GetNumberOfLODs() const override;
-	virtual void RetrieveRawMeshData(int32 LODIndex, FRawMesh& InOutRawMesh, bool bPropogateMeshData) const override;
+	virtual void RetrieveRawMeshData(int32 LODIndex, FMeshDescription& InOutRawMesh, bool bPropogateMeshData) const override;
 	virtual void RetrieveMeshSections(int32 LODIndex, TArray<FSectionInfo>& InOutSectionInfo) const override;
 	virtual int32 GetMaterialIndex(int32 LODIndex, int32 SectionIndex) const override;
 	virtual void ApplySettings(int32 LODIndex, FMeshData& InOutMeshData) const override;
diff --git a/Engine/Source/Developer/MeshMergeUtilities/Public/StaticMeshComponentAdapter.h b/Engine/Source/Developer/MeshMergeUtilities/Public/StaticMeshComponentAdapter.h
--- a/Engine/Source/Developer/MeshMergeUtilities/Public/StaticMeshComponentAdapter.h
+++ b/Engine/Source/Developer/MeshMergeUtilities/Public/StaticMeshComponentAdapter.h
@@ -15,7 +15,7 @@ public:
 
 	/** Begin IMaterialBakingAdapter overrides */
 	virtual int32 GetNumberOfLODs() const override;
-	virtual void RetrieveRawMeshData(int32 LODIndex, FRawMesh& InOutRawMesh, bool bPropogateMeshData) const override;
+	virtual void RetrieveRawMeshData(int32 LODIndex, FMeshDescription& InOutRawMesh, bool bPropogateMeshData) const override;
 	virtual void RetrieveMeshSections(int32 LODIndex, TArray<FSectionInfo>& InOutSectionInfo) const override;
 	virtual int32 GetMaterialIndex(int32 LODIndex, int32 SectionIndex) const override;	
 	virtual void ApplySettings(int32 LODIndex, FMeshData& InOutMeshData) const override;
diff --git a/Engine/Source/Developer/MeshReductionInterface/Public/IMeshReductionInterfaces.h b/Engine/Source/Developer/MeshReductionInterface/Public/IMeshReductionInterfaces.h
--- a/Engine/Source/Developer/MeshReductionInterface/Public/IMeshReductionInterfaces.h
+++ b/Engine/Source/Developer/MeshReductionInterface/Public/IMeshReductionInterfaces.h
@@ -70,7 +70,7 @@ public:
 	virtual bool IsSupported() const = 0;
 };
 
-DECLARE_DELEGATE_ThreeParams(FProxyCompleteDelegate, struct FRawMesh&, struct FFlattenMaterial&, const FGuid);
+DECLARE_DELEGATE_ThreeParams(FProxyCompleteDelegate, struct FMeshDescription&, struct FFlattenMaterial&, const FGuid);
 DECLARE_DELEGATE_TwoParams(FProxyFailedDelegate, const FGuid, const FString&);
 DECLARE_DELEGATE_TwoParams(FCreateProxyDelegate, const FGuid, TArray<UObject*>&);
 
diff --git a/Engine/Source/Developer/SimplygonMeshReduction/Private/SimplygonMeshReduction.cpp b/Engine/Source/Developer/SimplygonMeshReduction/Private/SimplygonMeshReduction.cpp
--- a/Engine/Source/Developer/SimplygonMeshReduction/Private/SimplygonMeshReduction.cpp
+++ b/Engine/Source/Developer/SimplygonMeshReduction/Private/SimplygonMeshReduction.cpp
@@ -1019,17 +1019,7 @@ public:
 				EdgeHardnesses[EdgeID] = false;
 			}
 
-			FRawMesh ConvertedRawMesh;
-			TMap<FName, int32> MaterialMap;
-			TPolygonGroupAttributesConstRef<FName> PolygonGroupNames = OutProxyMesh.PolygonGroupAttributes().GetAttributesRef<FName>(MeshAttribute::PolygonGroup::ImportedMaterialSlotName);
-			int32 MaterialIndex = 0;
-			for (const FPolygonGroupID PolygonGroupID : OutProxyMesh.PolygonGroups().GetElementIDs())
-			{
-				MaterialMap.Add(PolygonGroupNames[PolygonGroupID], MaterialIndex++);
-			}
-			FMeshDescriptionOperations::ConvertToRawMesh(OutProxyMesh, ConvertedRawMesh, MaterialMap);
-
-			Delegate.ExecuteIfBound(ConvertedRawMesh, OutMaterial, JobGUID);
+			Delegate.ExecuteIfBound(OutProxyMesh, OutMaterial, JobGUID);
 
 			return 0;
 		}
@@ -1217,18 +1207,8 @@ public:
 		if (OutProxyMesh.Polygons().Num() == 0)
 		{
 			FailedDelegate.ExecuteIfBound(InJobGUID, TEXT("Simplygon failed to generate a valid proxy mesh"));
-		}
-		FRawMesh ConvertedRawMesh;
-		TMap<FName, int32> MaterialMap;
-		TPolygonGroupAttributesConstRef<FName> PolygonGroupNames = OutProxyMesh.PolygonGroupAttributes().GetAttributesRef<FName>(MeshAttribute::PolygonGroup::ImportedMaterialSlotName);
-		int32 MaterialIndex = 0;
-		for (const FPolygonGroupID PolygonGroupID : OutProxyMesh.PolygonGroups().GetElementIDs())
-		{
-			MaterialMap.Add(PolygonGroupNames[PolygonGroupID], MaterialIndex++);
-		}
-		FMeshDescriptionOperations::ConvertToRawMesh(OutProxyMesh, ConvertedRawMesh, MaterialMap);
-				
-		CompleteDelegate.ExecuteIfBound(ConvertedRawMesh, OutMaterial, InJobGUID);
+		}				
+		CompleteDelegate.ExecuteIfBound(OutProxyMesh, OutMaterial, InJobGUID);
 	}
 
 		EventHandler.Task = nullptr;
diff --git a/Engine/Source/Developer/SimplygonSwarm/Private/SimplygonSwarm.cpp b/Engine/Source/Developer/SimplygonSwarm/Private/SimplygonSwarm.cpp
--- a/Engine/Source/Developer/SimplygonSwarm/Private/SimplygonSwarm.cpp
+++ b/Engine/Source/Developer/SimplygonSwarm/Private/SimplygonSwarm.cpp
@@ -208,7 +208,7 @@ public:
 
 		for (auto Data : InData)
 		{
-			InputMeshes.Push(Data.GetMeshDescription());
+			InputMeshes.Push(Data.RawMesh);
 		}
 
 		bool bDiscardEmissive = true;
@@ -373,17 +373,7 @@ public:
 			//if is bound then execute
 			if (CompleteDelegate.IsBound())
 			{
-				FRawMesh ConvertedRawMesh;
-				TMap<FName, int32> MaterialMap;
-				TPolygonGroupAttributesConstRef<FName> PolygonGroupNames = OutProxyMesh.PolygonGroupAttributes().GetAttributesRef<FName>(MeshAttribute::PolygonGroup::ImportedMaterialSlotName);
-				int32 MaterialIndex = 0;
-				for (const FPolygonGroupID PolygonGroupID : OutProxyMesh.PolygonGroups().GetElementIDs())
-				{
-					MaterialMap.Add(PolygonGroupNames[PolygonGroupID], MaterialIndex++);
-				}
-				FMeshDescriptionOperations::ConvertToRawMesh(OutProxyMesh, ConvertedRawMesh, MaterialMap);
-
-				CompleteDelegate.Execute(ConvertedRawMesh, OutMaterial, InSwarmTask.TaskData.ProcessorJobID);
+				CompleteDelegate.Execute(OutProxyMesh, OutMaterial, InSwarmTask.TaskData.ProcessorJobID);
 			}
 			else
 			{
@@ -593,7 +583,7 @@ private:
 			Count++;
 
 			//setup mesh data
-			ssf::pssfMeshData SsfMeshData = CreateSSFMeshDataFromRawMesh(*MergeData.GetMeshDescription(), MergeData.TexCoordBounds, MergeData.NewUVs);
+			ssf::pssfMeshData SsfMeshData = CreateSSFMeshDataFromRawMesh(*MergeData.RawMesh, MergeData.TexCoordBounds, MergeData.NewUVs);
 			SsfMesh->MeshDataList.push_back(SsfMeshData);
 
 			//setup mesh material information
@@ -602,7 +592,7 @@ private:
 			UniqueMaterialIds.Reserve(InputMaterials.Num());
 
 			//get unqiue material ids
-			GetUniqueMaterialIndices(*MergeData.GetMeshDescription(), UniqueMaterialIds);
+			GetUniqueMaterialIndices(*(MergeData.RawMesh), UniqueMaterialIds);
 
 			SsfMesh->MaterialIds->Items.reserve(UniqueMaterialIds.Num());
 
diff --git a/Engine/Source/Editor/UnrealEd/Public/Editor.h b/Engine/Source/Editor/UnrealEd/Public/Editor.h
--- a/Engine/Source/Editor/UnrealEd/Public/Editor.h
+++ b/Engine/Source/Editor/UnrealEd/Public/Editor.h
@@ -467,6 +467,7 @@ void MergeStaticMesh(UStaticMesh* DestMesh, UStaticMesh* SourceMesh, const FMerg
 // Converting models to static meshes.
 //
 UNREALED_API void GetBrushMesh(ABrush* Brush,UModel* Model,struct FRawMesh& OutMesh,TArray<FStaticMaterial>& OutMaterials);
+UNREALED_API void GetBrushMesh(ABrush* Brush, UModel* Model, struct FMeshDescription& OutMesh, TArray<FStaticMaterial>& OutMaterials);
 UStaticMesh* CreateStaticMeshFromBrush(UObject* Outer,FName Name,ABrush* Brush,UModel* Model);
  
 /**
diff --git a/Engine/Source/Runtime/MeshDescription/Public/MeshAttributeArray.h b/Engine/Source/Runtime/MeshDescription/Public/MeshAttributeArray.h
--- a/Engine/Source/Runtime/MeshDescription/Public/MeshAttributeArray.h
+++ b/Engine/Source/Runtime/MeshDescription/Public/MeshAttributeArray.h
@@ -106,6 +106,11 @@ public:
 		Insert( ElementCount - 1, Default );
 	}
 
+	uint32 GetHash(uint32 Crc = 0) const
+	{
+		return FCrc::MemCrc32(Container.GetData(), Container.Num() * sizeof(AttributeType), Crc);
+	}
+
 	/** Expands the array if necessary so that the passed element index is valid. Newly created elements will be assigned the default value. */
 	void Insert( const int32 Index, const AttributeType& Default );
 
@@ -245,6 +250,7 @@ public:
 	virtual void Insert( const int32 Index ) = 0;
 	virtual void Remove( const int32 Index ) = 0;
 	virtual void Initialize( const int32 Count ) = 0;
+	virtual uint32 GetHash() const = 0;
 	virtual void Serialize( FArchive& Ar ) = 0;
 	virtual void Remap( const TSparseArray<int32>& IndexRemap ) = 0;
 	virtual int32 GetNumIndices() const = 0;
@@ -333,6 +339,16 @@ public:
 		}
 	}
 
+	virtual uint32 GetHash() const
+	{
+		uint32 CrcResult = 0;
+		for (const TMeshAttributeArrayBase<AttributeType>& ArrayForIndex : ArrayForIndices)
+		{
+			CrcResult = ArrayForIndex.GetHash(CrcResult);
+		}
+		return CrcResult;
+	}
+
 	/** Polymorphic serialization */
 	virtual void Serialize( FArchive& Ar )
 	{
@@ -794,6 +810,15 @@ public:
 	/** Applies the given remapping to the attributes set */
 	void Remap( const TSparseArray<int32>& IndexRemap );
 
+	uint32 GetHash(const FName AttributeName) const
+	{
+		if (const FAttributesSetEntry* ArraySetPtr = Map.Find(AttributeName))
+		{
+			return (*ArraySetPtr)->GetHash();
+		}
+		return 0;
+	}
+
 	template <typename AttributeType>
 	DEPRECATED( 4.20, "Please use untemplated UnregisterAttribute() instead" )
 	void UnregisterAttribute( const FName AttributeName )
