Fix material baking crash
Fix the material baking automation tests

diff --git a/Engine/Source/Developer/MeshDescriptionOperations/Private/MeshDescriptionOperations.cpp b/Engine/Source/Developer/MeshDescriptionOperations/Private/MeshDescriptionOperations.cpp
--- a/Engine/Source/Developer/MeshDescriptionOperations/Private/MeshDescriptionOperations.cpp
+++ b/Engine/Source/Developer/MeshDescriptionOperations/Private/MeshDescriptionOperations.cpp
@@ -278,8 +278,22 @@ void FMeshDescriptionOperations::ConvertToRawMesh(const FMeshDescription& Source
 	DestinationRawMesh.FaceMaterialIndices.AddZeroed(TriangleNumber);
 	DestinationRawMesh.FaceSmoothingMasks.AddZeroed(TriangleNumber);
 
+	bool bHasVertexColor = false;
+	FVector4 WhiteColor(FLinearColor::White);
+	for (const FVertexInstanceID& VertexInstanceID : SourceMeshDescription.VertexInstances().GetElementIDs())
+	{
+		if (VertexInstanceColors[VertexInstanceID] != WhiteColor)
+		{
+			bHasVertexColor = true;
+			break;
+		}
+	}
+
 	int32 WedgeIndexNumber = TriangleNumber * 3;
-	DestinationRawMesh.WedgeColors.AddZeroed(WedgeIndexNumber);
+	if (bHasVertexColor)
+	{
+		DestinationRawMesh.WedgeColors.AddZeroed(WedgeIndexNumber);
+	}
 	DestinationRawMesh.WedgeIndices.AddZeroed(WedgeIndexNumber);
 	DestinationRawMesh.WedgeTangentX.AddZeroed(WedgeIndexNumber);
 	DestinationRawMesh.WedgeTangentY.AddZeroed(WedgeIndexNumber);
@@ -312,7 +326,10 @@ void FMeshDescriptionOperations::ConvertToRawMesh(const FMeshDescription& Source
 			{
 				const FVertexInstanceID VertexInstanceID = MeshTriangle.GetVertexInstanceID(Corner);
 
-				DestinationRawMesh.WedgeColors[WedgeIndex] = FLinearColor(VertexInstanceColors[VertexInstanceID]).ToFColor(true);
+				if (bHasVertexColor)
+				{
+					DestinationRawMesh.WedgeColors[WedgeIndex] = FLinearColor(VertexInstanceColors[VertexInstanceID]).ToFColor(true);
+				}
 				DestinationRawMesh.WedgeIndices[WedgeIndex] = RemapVerts[SourceMeshDescription.GetVertexInstanceVertex(VertexInstanceID).GetValue()];
 				DestinationRawMesh.WedgeTangentX[WedgeIndex] = VertexInstanceTangents[VertexInstanceID];
 				DestinationRawMesh.WedgeTangentY[WedgeIndex] = FVector::CrossProduct(VertexInstanceNormals[VertexInstanceID], VertexInstanceTangents[VertexInstanceID]).GetSafeNormal() * VertexInstanceBinormalSigns[VertexInstanceID];
diff --git a/Engine/Source/Developer/MeshMergeUtilities/Private/MeshMergeUtilities.cpp b/Engine/Source/Developer/MeshMergeUtilities/Private/MeshMergeUtilities.cpp
--- a/Engine/Source/Developer/MeshMergeUtilities/Private/MeshMergeUtilities.cpp
+++ b/Engine/Source/Developer/MeshMergeUtilities/Private/MeshMergeUtilities.cpp
@@ -139,6 +139,7 @@ void FMeshMergeUtilities::BakeMaterialsForComponent(TArray<TWeakObjectPtr<UObjec
 		if (bProcessedLOD)
 		{
 			FMeshDescription& RawMesh = RawMeshLODs.Add(LODIndex);
+			UStaticMesh::RegisterMeshAttributes(RawMesh);
 			Adapter->RetrieveRawMeshData(LODIndex, RawMesh, MaterialOptions->bUseMeshData);
 		}
 
@@ -187,6 +188,7 @@ void FMeshMergeUtilities::BakeMaterialsForComponent(TArray<TWeakObjectPtr<UObjec
 				UniqueSectionIndexPerLOD.MultiFind(LODIndex, IndexPairs);
 
 				FMeshData MeshSettings;
+				MeshSettings.RawMesh = new FRawMesh();
 
 				// Add material indices used for rendering out material
 				for (const TPair<uint32, uint32>& Pair : IndexPairs)
@@ -200,10 +202,9 @@ void FMeshMergeUtilities::BakeMaterialsForComponent(TArray<TWeakObjectPtr<UObjec
 				if (MeshSettings.MaterialIndices.Num())
 				{
 					// Retrieve raw mesh
-					const FMeshDescription* MeshDescription = RawMeshLODs.Find(LODIndex);
+					MeshSettings.RawMeshDescription = RawMeshLODs.Find(LODIndex);
 					TMap<FName, int32> MaterialMap;
-					FMeshDescriptionOperations::ConvertToRawMesh(*MeshDescription, *(MeshSettings.RawMesh), MaterialMap);
-					
+					FMeshDescriptionOperations::ConvertToRawMesh(*MeshSettings.RawMeshDescription, *(MeshSettings.RawMesh), MaterialMap);
 
 					MeshSettings.TextureCoordinateBox = FBox2D(FVector2D(0.0f, 0.0f), FVector2D(1.0f, 1.0f));
 					const bool bUseVertexColor = (MeshSettings.RawMesh->WedgeColors.Num() > 0);
@@ -278,6 +279,7 @@ void FMeshMergeUtilities::BakeMaterialsForComponent(TArray<TWeakObjectPtr<UObjec
 			if (MeshSettings.MaterialIndices.Num())
 			{
 				MeshSettings.RawMesh = nullptr;
+				MeshSettings.RawMeshDescription = nullptr;
 				MeshSettings.TextureCoordinateBox = FBox2D(FVector2D(0.0f, 0.0f), FVector2D(1.0f, 1.0f));
 				MeshSettings.TextureCoordinateIndex = 0;
 
@@ -437,6 +439,15 @@ void FMeshMergeUtilities::BakeMaterialsForComponent(TArray<TWeakObjectPtr<UObjec
 	}
 
 	Adapter->UpdateUVChannelData();
+	//Free the FRawMesh memory
+	for (FMeshData& MeshData : GlobalMeshSettings)
+	{
+		if (MeshData.RawMesh != nullptr)
+		{
+			delete MeshData.RawMesh;
+			MeshData.RawMesh = nullptr;
+		}
+	}
 }
 
 void FMeshMergeUtilities::BakeMaterialsForComponent(USkeletalMeshComponent* SkeletalMeshComponent) const
