Fix crash, avoid deleting a pointer 2 time.

diff --git a/Engine/Source/Developer/MeshMergeUtilities/Private/MeshMergeUtilities.cpp b/Engine/Source/Developer/MeshMergeUtilities/Private/MeshMergeUtilities.cpp
--- a/Engine/Source/Developer/MeshMergeUtilities/Private/MeshMergeUtilities.cpp
+++ b/Engine/Source/Developer/MeshMergeUtilities/Private/MeshMergeUtilities.cpp
@@ -188,7 +188,7 @@ void FMeshMergeUtilities::BakeMaterialsForComponent(TArray<TWeakObjectPtr<UObjec
 				UniqueSectionIndexPerLOD.MultiFind(LODIndex, IndexPairs);
 
 				FMeshData MeshSettings;
-				MeshSettings.RawMeshDescription = new FMeshDescription();
+				MeshSettings.RawMeshDescription = nullptr;
 
 				// Add material indices used for rendering out material
 				for (const TPair<uint32, uint32>& Pair : IndexPairs)
@@ -203,6 +203,9 @@ void FMeshMergeUtilities::BakeMaterialsForComponent(TArray<TWeakObjectPtr<UObjec
 				{
 					// Retrieve raw mesh
 					MeshSettings.RawMeshDescription = RawMeshLODs.Find(LODIndex);
+					
+					//Should not be using mesh data if there is no mesh
+					check(MeshSettings.RawMeshDescription);
 
 					MeshSettings.TextureCoordinateBox = FBox2D(FVector2D(0.0f, 0.0f), FVector2D(1.0f, 1.0f));
 					const bool bUseVertexColor = FMeshDescriptionOperations::HasVertexColor(*(MeshSettings.RawMeshDescription));
@@ -438,15 +441,7 @@ void FMeshMergeUtilities::BakeMaterialsForComponent(TArray<TWeakObjectPtr<UObjec
 	}
 
 	Adapter->UpdateUVChannelData();
-	//Free the FMeshDescription memory
-	for (FMeshData& MeshData : GlobalMeshSettings)
-	{
-		if (MeshData.RawMeshDescription != nullptr)
-		{
-			delete MeshData.RawMeshDescription;
-			MeshData.RawMeshDescription = nullptr;
-		}
-	}
+	GlobalMeshSettings.Empty();
 }
 
 void FMeshMergeUtilities::BakeMaterialsForComponent(USkeletalMeshComponent* SkeletalMeshComponent) const
