Merge Actors creates a material with twice the prefix (M_M_)

diff --git a/Engine/Source/Developer/MeshMergeUtilities/Private/MeshMergeUtilities.cpp b/Engine/Source/Developer/MeshMergeUtilities/Private/MeshMergeUtilities.cpp
--- a/Engine/Source/Developer/MeshMergeUtilities/Private/MeshMergeUtilities.cpp
+++ b/Engine/Source/Developer/MeshMergeUtilities/Private/MeshMergeUtilities.cpp
@@ -3077,12 +3077,12 @@ UMaterialInterface* FMeshMergeUtilities::CreateProxyMaterial(const FString &InBa
 	FString MaterialPackageName;
 	if (InBasePackageName.IsEmpty())
 	{
-		MaterialAssetName = TEXT("M_MERGED_") + FPackageName::GetShortName(MergedAssetPackageName);
+		MaterialAssetName = FPackageName::GetShortName(MergedAssetPackageName);
 		MaterialPackageName = FPackageName::GetLongPackagePath(MergedAssetPackageName) + TEXT("/") + MaterialAssetName;
 	}
 	else
 	{
-		MaterialAssetName = TEXT("M_") + FPackageName::GetShortName(InBasePackageName);
+		MaterialAssetName = FPackageName::GetShortName(InBasePackageName);
 		MaterialPackageName = FPackageName::GetLongPackagePath(InBasePackageName) + TEXT("/") + MaterialAssetName;
 	}
 
@@ -3095,7 +3095,7 @@ UMaterialInterface* FMeshMergeUtilities::CreateProxyMaterial(const FString &InBa
 		MaterialPackage->Modify();
 	}
 
-	UMaterialInstanceConstant* MergedMaterial = ProxyMaterialUtilities::CreateProxyMaterialInstance(MaterialPackage, InSettings.MaterialSettings, InBaseMaterial, OutMaterial, MaterialAssetName, MaterialPackageName, OutAssetsToSync);
+	UMaterialInstanceConstant* MergedMaterial = ProxyMaterialUtilities::CreateProxyMaterialInstance(MaterialPackage, InSettings.MaterialSettings, InBaseMaterial, OutMaterial, MaterialPackageName, MaterialAssetName, OutAssetsToSync);
 	// Set material static lighting usage flag if project has static lighting enabled
 	static const auto AllowStaticLightingVar = IConsoleManager::Get().FindTConsoleVariableDataInt(TEXT("r.AllowStaticLighting"));
 	const bool bAllowStaticLighting = (!AllowStaticLightingVar || AllowStaticLightingVar->GetValueOnGameThread() != 0);
