ProxyLOD.  Correcting a crash when a landscape object was present and culling was enabled.   Regression caused by FMeshDescriptor conversion.
P1

diff --git a/Engine/Source/Developer/MeshMergeUtilities/Private/MeshMergeHelpers.cpp b/Engine/Source/Developer/MeshMergeUtilities/Private/MeshMergeHelpers.cpp
--- a/Engine/Source/Developer/MeshMergeUtilities/Private/MeshMergeHelpers.cpp
+++ b/Engine/Source/Developer/MeshMergeUtilities/Private/MeshMergeHelpers.cpp
@@ -911,13 +911,13 @@ void FMeshMergeHelpers::RetrieveCullingLandscapeAndVolumes(UWorld* InWorld, cons
 	for (ALandscapeProxy* Landscape : LandscapeActors)
 	{
 		// Export the landscape to raw mesh format
-		FMeshDescription MeshDescription;
-		UStaticMesh::RegisterMeshAttributes(MeshDescription);
+		FMeshDescription* MeshDescription = new FMeshDescription();
+		UStaticMesh::RegisterMeshAttributes(*MeshDescription);
 		FBoxSphereBounds LandscapeBounds = EstimatedMeshProxyBounds;
-		Landscape->ExportToRawMesh(LandscapeExportLOD, MeshDescription, LandscapeBounds);
-		if (MeshDescription.Vertices().Num())
+		Landscape->ExportToRawMesh(LandscapeExportLOD, *MeshDescription, LandscapeBounds);
+		if (MeshDescription->Vertices().Num())
 		{
-			CullingRawMeshes.Add(&MeshDescription);
+			CullingRawMeshes.Add(MeshDescription);
 		}
 	}
 
diff --git a/Engine/Source/Developer/MeshMergeUtilities/Private/MeshMergeHelpers.h b/Engine/Source/Developer/MeshMergeUtilities/Private/MeshMergeHelpers.h
--- a/Engine/Source/Developer/MeshMergeUtilities/Private/MeshMergeHelpers.h
+++ b/Engine/Source/Developer/MeshMergeUtilities/Private/MeshMergeHelpers.h
@@ -54,7 +54,7 @@ public:
 	/** Transforms raw mesh data using InTransform*/
 	static void TransformRawMeshVertexData(const FTransform& InTransform, FMeshDescription &OutRawMesh);
 
-	/** Retrieves all culling landscapes and volumes as FMeshDescription structures */
+	/** Retrieves all culling landscapes and volumes as FMeshDescription structures. Note the caller is responsible for deleting the heap data managed by CullingRawMeshes */
 	static void RetrieveCullingLandscapeAndVolumes(UWorld* InWorld, const FBoxSphereBounds& EstimatedMeshProxyBounds, const TEnumAsByte<ELandscapeCullingPrecision::Type> PrecisionType, TArray<FMeshDescription*>& CullingRawMeshes);
 
 	/** Transforms physics geometry data using InTransform */
diff --git a/Engine/Source/Developer/MeshMergeUtilities/Private/MeshMergeUtilities.cpp b/Engine/Source/Developer/MeshMergeUtilities/Private/MeshMergeUtilities.cpp
--- a/Engine/Source/Developer/MeshMergeUtilities/Private/MeshMergeUtilities.cpp
+++ b/Engine/Source/Developer/MeshMergeUtilities/Private/MeshMergeUtilities.cpp
@@ -1484,7 +1484,7 @@ void FMeshMergeUtilities::CreateProxyMesh(const TArray<UStaticMeshComponent*>& I
 		}
 	};
 
-	// Landscape culling
+	// Landscape culling.  NB these are temporary copies of the culling data and should be deleted after use.
 	TArray<FMeshDescription*> CullingRawMeshes;
 	if (InMeshProxySettings.bUseLandscapeCulling)
 	{
@@ -1587,8 +1587,15 @@ void FMeshMergeUtilities::CreateProxyMesh(const TArray<UStaticMeshComponent*>& I
 
 		Processor->Tick(0); // make sure caller gets merging results
 	}
+
+	// Clean up the CullingRawMeshes
+	for (FMeshDescription* RawMesh : CullingRawMeshes)
+	{
+		delete RawMesh;
+	}
 }
 
+
 bool FMeshMergeUtilities::IsValidBaseMaterial(const UMaterialInterface* InBaseMaterial, bool bShowToaster) const
 {
 	if (InBaseMaterial != nullptr)
