Fix MeshMergeUtilities when remapping polygon group

diff --git a/Engine/Source/Developer/MeshMergeUtilities/Private/MeshMergeUtilities.cpp b/Engine/Source/Developer/MeshMergeUtilities/Private/MeshMergeUtilities.cpp
--- a/Engine/Source/Developer/MeshMergeUtilities/Private/MeshMergeUtilities.cpp
+++ b/Engine/Source/Developer/MeshMergeUtilities/Private/MeshMergeUtilities.cpp
@@ -1456,31 +1456,7 @@ void FMeshMergeUtilities::CreateProxyMesh(const TArray<UStaticMeshComponent*>& I
 				RemapPolygonGroup.Add(PolygonGroupID, FPolygonGroupID(RemapID));
 				MaxRemapID = FMath::Max(MaxRemapID, RemapID);
 			}
-
-			int32 PolygonGroupRemapCount = FMath::Max(MaxRemapID, RemapPolygonGroup.Num());
-			TSparseArray<int32> PolygonGroupRemap;
-			PolygonGroupRemap.Reserve(PolygonGroupRemapCount);
-			for (int32 Index = 0; Index < PolygonGroupRemapCount; ++Index)
-			{
-				PolygonGroupRemap.AddUninitialized();
-			}
-			//Set the polygon state
-			for (auto Kvp : RemapPolygonGroup)
-			{
-				PolygonGroupRemap[Kvp.Key.GetValue()] = Kvp.Value.GetValue();
-				FMeshPolygonGroup& PolygonGroup = RawMesh.GetPolygonGroup(Kvp.Key);
-				for (FPolygonID PolygonID : PolygonGroup.Polygons)
-				{
-					FMeshPolygon& Polygon = RawMesh.GetPolygon(PolygonID);
-					Polygon.PolygonGroupID = Kvp.Value;
-				}
-			}
-
-			//Remap the polygon groups
-			for (auto Kvp : RemapPolygonGroup)
-			{
-				RawMesh.PolygonGroups().Remap(PolygonGroupRemap);
-			}
+			FMeshDescriptionOperations::RemapPolygonGroups(RawMesh, RemapPolygonGroup);
 		}
 	};
 
