Textures and Materials have their merged actor's name appended to their names

diff --git a/Engine/Source/Developer/MeshMergeUtilities/Private/MeshMergeUtilities.cpp b/Engine/Source/Developer/MeshMergeUtilities/Private/MeshMergeUtilities.cpp
--- a/Engine/Source/Developer/MeshMergeUtilities/Private/MeshMergeUtilities.cpp
+++ b/Engine/Source/Developer/MeshMergeUtilities/Private/MeshMergeUtilities.cpp
@@ -3061,18 +3061,18 @@ UMaterialInterface* FMeshMergeUtilities::CreateProxyMaterial(const FString &InBa
 	if (InBasePackageName.IsEmpty())
 	{
 		MaterialAssetName = FPackageName::GetShortName(MergedAssetPackageName);
-		MaterialPackageName = FPackageName::GetLongPackagePath(MergedAssetPackageName) + TEXT("/") + MaterialAssetName;
+		MaterialPackageName = FPackageName::GetLongPackagePath(MergedAssetPackageName) + TEXT("/");
 	}
 	else
 	{
 		MaterialAssetName = FPackageName::GetShortName(InBasePackageName);
-		MaterialPackageName = FPackageName::GetLongPackagePath(InBasePackageName) + TEXT("/") + MaterialAssetName;
+		MaterialPackageName = FPackageName::GetLongPackagePath(InBasePackageName) + TEXT("/");
 	}
 
 	UPackage* MaterialPackage = InOuter;
 	if (MaterialPackage == nullptr)
 	{
-		MaterialPackage = CreatePackage(nullptr, *MaterialPackageName);
+		MaterialPackage = CreatePackage(nullptr, *(MaterialPackageName + MaterialAssetName));
 		check(MaterialPackage);
 		MaterialPackage->FullyLoad();
 		MaterialPackage->Modify();
