PR #5167: Fix infinite loop while building HLOD proxy meshes (Contributed by [REDACTED_PERSON])

diff --git a/Engine/Source/Developer/MeshMergeUtilities/Private/MeshMergeUtilities.cpp b/Engine/Source/Developer/MeshMergeUtilities/Private/MeshMergeUtilities.cpp
--- a/Engine/Source/Developer/MeshMergeUtilities/Private/MeshMergeUtilities.cpp
+++ b/Engine/Source/Developer/MeshMergeUtilities/Private/MeshMergeUtilities.cpp
@@ -1133,6 +1133,10 @@ void FMeshMergeUtilities::CreateProxyMesh(const TArray<UStaticMeshComponent*>& I
 	if (ComponentsToMerge.Num() == 0)
 	{
 		UE_LOG(LogMeshMerging, Log, TEXT("No static mesh specified to generate a proxy mesh for"));
+		
+		TArray<UObject*> OutAssetsToSync;
+		InProxyCreatedDelegate.ExecuteIfBound(InGuid, OutAssetsToSync);
+
 		return;
 	}
 
diff --git a/Engine/Source/Developer/MeshMergeUtilities/Private/ProxyGenerationProcessor.cpp b/Engine/Source/Developer/MeshMergeUtilities/Private/ProxyGenerationProcessor.cpp
--- a/Engine/Source/Developer/MeshMergeUtilities/Private/ProxyGenerationProcessor.cpp
+++ b/Engine/Source/Developer/MeshMergeUtilities/Private/ProxyGenerationProcessor.cpp
@@ -111,6 +111,9 @@ void FProxyGenerationProcessor::ProxyGenerationFailed(const FGuid OutJobGUID, co
 	{
 		UE_LOG(LogMeshMerging, Log, TEXT("Failed to generate proxy mesh for cluster %s, %s"), *(*FindData)->ProxyBasePackageName, *ErrorMessage);
 		ProxyMeshJobs.Remove(OutJobGUID);
+		
+		TArray<UObject*> OutAssetsToSync;
+		(*FindData)->CallbackDelegate.ExecuteIfBound(OutJobGUID, OutAssetsToSync);
 	}
 }
 
