Fix mesh merge utilities LOD extra material

diff --git a/Engine/Source/Developer/MeshMergeUtilities/Private/MeshMergeUtilities.cpp b/Engine/Source/Developer/MeshMergeUtilities/Private/MeshMergeUtilities.cpp
--- a/Engine/Source/Developer/MeshMergeUtilities/Private/MeshMergeUtilities.cpp
+++ b/Engine/Source/Developer/MeshMergeUtilities/Private/MeshMergeUtilities.cpp
@@ -2685,6 +2685,21 @@ void FMeshMergeUtilities::CreateMergedRawMeshes(FMeshMergeDataTracker& InDataTra
 					FMeshDescriptionOperations::AppendMeshDescription(*RawMeshPtr, MergedMesh, AppendSettings);
 				}
 			}
+
+			//Cleanup the empty material to avoid empty section later
+			TArray<FPolygonGroupID> PolygonGroupToRemove;
+			for (FPolygonGroupID PolygonGroupID : MergedMesh.PolygonGroups().GetElementIDs())
+			{
+				if (MergedMesh.GetPolygonGroupPolygons(PolygonGroupID).Num() < 1)
+				{
+					PolygonGroupToRemove.Add(PolygonGroupID);
+					
+				}
+			}
+			for (FPolygonGroupID PolygonGroupID : PolygonGroupToRemove)
+			{
+				MergedMesh.DeletePolygonGroup(PolygonGroupID);
+			}
 		}
 	}
 	else
