Fix crash on edit after ProxyLOD merge

diff --git a/Engine/Source/Developer/MeshMergeUtilities/Private/ProxyGenerationProcessor.cpp b/Engine/Source/Developer/MeshMergeUtilities/Private/ProxyGenerationProcessor.cpp
--- a/Engine/Source/Developer/MeshMergeUtilities/Private/ProxyGenerationProcessor.cpp
+++ b/Engine/Source/Developer/MeshMergeUtilities/Private/ProxyGenerationProcessor.cpp
@@ -208,7 +208,7 @@ void FProxyGenerationProcessor::ProcessJob(const FGuid& JobGuid, FProxyGeneratio
 	const bool bContainsImposters = Data->MergeData->ImposterComponents.Num() > 0;
 	FBox ImposterBounds(EForceInit::ForceInit);
 
-	auto RemoveVertexColorAndCommitMeshDescription = [&StaticMesh, &Data]()
+	auto RemoveVertexColorAndCommitMeshDescription = [&StaticMesh, &Data, &ProxyMaterial]()
 	{
 		if (!Data->MergeData->InProxySettings.bAllowVertexColors)
 		{
@@ -227,6 +227,16 @@ void FProxyGenerationProcessor::ProcessJob(const FGuid& JobGuid, FProxyGeneratio
 		if (ensure(MeshDescription))
 		{
 			*MeshDescription = Data->RawMesh;
+
+			// material index / name mapping
+			TPolygonGroupAttributesConstRef<FName> PolygonGroupMaterialSlotName = MeshDescription->PolygonGroupAttributes().GetAttributesRef<FName>(MeshAttribute::PolygonGroup::ImportedMaterialSlotName);
+			for (const FPolygonGroupID PolygonGroupID : MeshDescription->PolygonGroups().GetElementIDs())
+			{
+				FStaticMaterial NewMaterial(ProxyMaterial);
+				NewMaterial.ImportedMaterialSlotName = PolygonGroupMaterialSlotName[PolygonGroupID];
+				StaticMesh->StaticMaterials.Add(NewMaterial);
+			}
+
 			StaticMesh->CommitMeshDescription(SourceModelIndex);
 		}
 	};
@@ -239,7 +249,6 @@ void FProxyGenerationProcessor::ProcessJob(const FGuid& JobGuid, FProxyGeneratio
 		// The base material index is always one here as we assume we only have one HLOD material
 		FMeshMergeHelpers::MergeImpostersToRawMesh(Data->MergeData->ImposterComponents, Data->RawMesh, FVector::ZeroVector, 1, ImposterMaterials);
 
-		
 		for (const UStaticMeshComponent* Component : Data->MergeData->ImposterComponents)
 		{
 			if (Component->GetStaticMesh())
@@ -249,8 +258,6 @@ void FProxyGenerationProcessor::ProcessJob(const FGuid& JobGuid, FProxyGeneratio
 		}
 		RemoveVertexColorAndCommitMeshDescription();
 
-		StaticMesh->StaticMaterials.Add(FStaticMaterial(ProxyMaterial));
-
 		for (UMaterialInterface* Material : ImposterMaterials)
 		{
 			StaticMesh->StaticMaterials.Add(FStaticMaterial(Material));
@@ -259,9 +266,7 @@ void FProxyGenerationProcessor::ProcessJob(const FGuid& JobGuid, FProxyGeneratio
 	else
 	{
 		RemoveVertexColorAndCommitMeshDescription();
-
-		StaticMesh->StaticMaterials.Add(FStaticMaterial(ProxyMaterial));
-	}	
+	}
 
 	//Set the Imported version before calling the build
 	StaticMesh->ImportVersion = EImportStaticMeshVersion::LastVersion;
