Added the option "Create Merged Material" to mesh merging settings
* Option is disabled by default
* When enabled, the resulting merged mesh will have:
** An additional unused material which is created from the merge of all meshes materials
** An additional set of UV coordinates that can be used to apply the merged material to any section of the mesh
** Added additional static switch params to BaseFlattenMaterials to allow specifying UV channel
#rnx

diff --git a/Engine/Source/Developer/MaterialUtilities/Public/MaterialUtilities.h b/Engine/Source/Developer/MaterialUtilities/Public/MaterialUtilities.h
--- a/Engine/Source/Developer/MaterialUtilities/Public/MaterialUtilities.h
+++ b/Engine/Source/Developer/MaterialUtilities/Public/MaterialUtilities.h
@@ -56,6 +56,7 @@ struct FFlattenMaterial
 		, bDitheredLODTransition(false)
 		, BlendMode(BLEND_Opaque)
 		, EmissiveScale(1.0f)
+		, UVChannel(0)
 	{
 		for (FIntPoint& Size : PropertySizes)
 		{
@@ -108,8 +109,10 @@ struct FFlattenMaterial
 	EBlendMode		BlendMode;
 	/** Scale (maximum baked down value) for the emissive property */
 	float			EmissiveScale;
-private:
+	/** UV channel to use */
+	int32			UVChannel;
 
+private:
 	/** Texture sizes for each individual property*/
 	FIntPoint PropertySizes[(uint32)EFlattenMaterialProperties::NumFlattenMaterialProperties];
 	/** Baked down texture samples for each individual property*/
diff --git a/Engine/Source/Developer/MeshMergeUtilities/Private/MeshMergeUtilities.cpp b/Engine/Source/Developer/MeshMergeUtilities/Private/MeshMergeUtilities.cpp
--- a/Engine/Source/Developer/MeshMergeUtilities/Private/MeshMergeUtilities.cpp
+++ b/Engine/Source/Developer/MeshMergeUtilities/Private/MeshMergeUtilities.cpp
@@ -1942,6 +1942,12 @@ void FMeshMergeUtilities::MergeComponentsToStaticMesh(const TArray<UPrimitiveCom
 	TArray<float> MaterialImportanceValues;
 	FMaterialUtilities::DetermineMaterialImportance(UniqueMaterials, MaterialImportanceValues);
 
+	TMultiMap< FMeshLODKey, MaterialRemapPair > OutputMaterialsMap;
+
+	// The UV channel to use for the flattened material
+	int32 MergedMatUVChannel = INDEX_NONE;
+	UMaterialInterface* MergedMaterial = nullptr;
+
 	// If the user wants to merge materials into a single one
 	if (bMergeMaterialData)
 	{
@@ -1954,8 +1960,6 @@ void FMeshMergeUtilities::MergeComponentsToStaticMesh(const TArray<UPrimitiveCom
 		TArray<FMaterialData> GlobalMaterialSettings;
 		TArray<float> SectionMaterialImportanceValues;
 
-		TMultiMap< FMeshLODKey, MaterialRemapPair > OutputMaterialsMap;
-
 		TMap<EMaterialProperty, FIntPoint> PropertySizes;
 		for (const FPropertyEntry& Entry : MaterialOptions->Properties)
 		{
@@ -2009,7 +2013,7 @@ void FMeshMergeUtilities::MergeComponentsToStaticMesh(const TArray<UPrimitiveCom
 				MeshData.bMirrored = Component->GetComponentTransform().GetDeterminant() < 0.0f;
 				int32 MeshDataIndex = 0;
 
-				if (InSettings.bUseVertexDataForBakingMaterial && (bDoesMaterialUseVertexData || bRequiresUniqueUVs))
+				if (InSettings.bCreateMergedMaterial || (InSettings.bUseVertexDataForBakingMaterial && (bDoesMaterialUseVertexData || bRequiresUniqueUVs)))
 				{
 					MeshData.RawMeshDescription = DataTracker.GetRawMeshPtr(Key);
 
@@ -2255,6 +2259,21 @@ void FMeshMergeUtilities::MergeComponentsToStaticMesh(const TArray<UPrimitiveCom
 			}
 		}
 
+		// Compute UV channel to use for the merged material
+		if(InSettings.bCreateMergedMaterial)
+		{
+			for(TConstRawMeshIterator Iterator = DataTracker.GetConstRawMeshIterator(); Iterator; ++Iterator)
+			{
+				const FMeshDescription& RawMesh = Iterator.Value();
+
+				if(RawMesh.Vertices().Num())
+				{
+					const TVertexInstanceAttributesRef<const FVector2D> VertexInstanceUVs = RawMesh.VertexInstanceAttributes().GetAttributesRef<FVector2D>(MeshAttribute::VertexInstance::TextureCoordinate);
+					MergedMatUVChannel = FMath::Max(MergedMatUVChannel, VertexInstanceUVs.GetNumIndices());
+				}
+			}
+		}
+
 		// Adjust UVs
 		for (int32 ComponentIndex = 0; ComponentIndex < StaticMeshComponentsToMerge.Num(); ++ComponentIndex)
 		{
@@ -2263,9 +2282,6 @@ void FMeshMergeUtilities::MergeComponentsToStaticMesh(const TArray<UPrimitiveCom
 			{
 				if (MappingPair.Key.GetMeshIndex() == ComponentIndex && !ProcessedMaterials.Contains(MappingPair.Value.Key))
 				{
-					const int32 LODIndex = MappingPair.Key.GetLODIndex();
-					// Found component entry
-
 					// Retrieve raw mesh data for this component and lod pair
 					FMeshDescription* RawMesh = DataTracker.GetRawMeshPtr(MappingPair.Key);
 
@@ -2284,7 +2300,7 @@ void FMeshMergeUtilities::MergeComponentsToStaticMesh(const TArray<UPrimitiveCom
 							for (FVertexInstanceID VertexInstanceID : RawMesh->VertexInstances().GetElementIDs())
 							{
 								FVector2D UV = VertexInstanceUVs.Get(VertexInstanceID, UVChannelIdx);
-								if (UVChannelIdx == 0)
+								if (UVChannelIdx == 0 && !InSettings.bCreateMergedMaterial)
 								{
 									if (MeshData.CustomTextureCoordinates.Num())
 									{
@@ -2313,6 +2329,19 @@ void FMeshMergeUtilities::MergeComponentsToStaticMesh(const TArray<UPrimitiveCom
 								VertexIndex++;
 							}
 						}
+
+						if (InSettings.bCreateMergedMaterial)
+						{
+							VertexInstanceUVs.SetNumIndices(MergedMatUVChannel + 1);
+
+							int32 VertexIndex = 0;
+							for(FVertexInstanceID VertexInstanceID : RawMesh->VertexInstances().GetElementIDs())
+							{
+								FVector2D UV = MeshData.CustomTextureCoordinates[VertexIndex];
+								VertexInstanceUVs.Set(VertexInstanceID, MergedMatUVChannel, UV * UVTransform.Value + UVTransform.Key);
+								VertexIndex++;
+							}
+						}
 					}
 				}
 			}
@@ -2330,14 +2359,20 @@ void FMeshMergeUtilities::MergeComponentsToStaticMesh(const TArray<UPrimitiveCom
 			FMeshDescriptionOperations::RemapPolygonGroups(RawMesh, RemapPolygonGroups);
 		}
 
-		UMaterialInterface* MergedMaterial = CreateProxyMaterial(InBasePackageName, MergedAssetPackageName, InBaseMaterial, InOuter, InSettings, OutMaterial, OutAssetsToSync);
-		UniqueMaterials.Empty(1);
-		UniqueMaterials.Add(MergedMaterial);
-		
-		FSectionInfo NewSection;
-		NewSection.Material = MergedMaterial;
-		NewSection.EnabledProperties.Add(GET_MEMBER_NAME_CHECKED(FStaticMeshSection, bCastShadow));
-		DataTracker.AddBakedMaterialSection(NewSection);
+		OutMaterial.UVChannel = MergedMatUVChannel;
+
+		MergedMaterial = CreateProxyMaterial(InBasePackageName, MergedAssetPackageName, InBaseMaterial, InOuter, InSettings, OutMaterial, OutAssetsToSync);
+
+		if (!InSettings.bCreateMergedMaterial)
+		{
+			UniqueMaterials.Empty(1);
+			UniqueMaterials.Add(MergedMaterial);
+
+			FSectionInfo NewSection;
+			NewSection.Material = MergedMaterial;
+			NewSection.EnabledProperties.Add(GET_MEMBER_NAME_CHECKED(FStaticMeshSection, bCastShadow));
+			DataTracker.AddBakedMaterialSection(NewSection);
+		}
 
 		for (IMeshMergeExtension* Extension : MeshMergeExtensions)
 		{
@@ -2346,8 +2381,7 @@ void FMeshMergeUtilities::MergeComponentsToStaticMesh(const TArray<UPrimitiveCom
 	}
 
 	TArray<FMeshDescription> MergedRawMeshes;
-	TMultiMap<FMeshLODKey, MaterialRemapPair> OutputMaterialsMap;
-	CreateMergedRawMeshes(DataTracker, InSettings, StaticMeshComponentsToMerge, UniqueMaterials, CollapsedMaterialMap, OutputMaterialsMap, bMergeAllLODs, bMergeMaterialData, MergedAssetPivot, MergedRawMeshes);
+	CreateMergedRawMeshes(DataTracker, InSettings, StaticMeshComponentsToMerge, UniqueMaterials, CollapsedMaterialMap, OutputMaterialsMap, bMergeAllLODs, bMergeMaterialData && !InSettings.bCreateMergedMaterial, MergedAssetPivot, MergedRawMeshes);
 
 	// Populate mesh section map
 	FMeshSectionInfoMap SectionInfoMap;	
@@ -2386,6 +2420,11 @@ void FMeshMergeUtilities::MergeComponentsToStaticMesh(const TArray<UPrimitiveCom
 		}
 	}
 
+	if(InSettings.bCreateMergedMaterial)
+	{
+		OutputMaterialsMap.Reset();
+	}
+
 	// Transform physics primitives to merged mesh pivot
 	if (InSettings.bMergePhysicsData && !MergedAssetPivot.IsZero())
 	{
@@ -2597,6 +2636,19 @@ void FMeshMergeUtilities::MergeComponentsToStaticMesh(const TArray<UPrimitiveCom
 
 		StaticMesh->PostEditChange();
 
+		if (InSettings.bCreateMergedMaterial)
+		{
+			//Make sure we have unique slot name here
+			FName MaterialSlotName = MergedMaterial->GetFName();
+			int32 Counter = 1;
+			while (!IsMaterialImportedNameUnique(MaterialSlotName))
+			{
+				MaterialSlotName = *(MergedMaterial->GetName() + TEXT("_") + FString::FromInt(Counter++));
+			}
+			StaticMesh->StaticMaterials.Add(FStaticMaterial(MergedMaterial, MaterialSlotName));
+			StaticMesh->UpdateUVChannelData(false);
+		}
+
 		OutAssetsToSync.Add(StaticMesh);
 		OutMergedActorLocation = MergedAssetPivot;
 	}
diff --git a/Engine/Source/Developer/MeshMergeUtilities/Private/ProxyMaterialUtilities.h b/Engine/Source/Developer/MeshMergeUtilities/Private/ProxyMaterialUtilities.h
--- a/Engine/Source/Developer/MeshMergeUtilities/Private/ProxyMaterialUtilities.h
+++ b/Engine/Source/Developer/MeshMergeUtilities/Private/ProxyMaterialUtilities.h
@@ -102,7 +102,7 @@ namespace ProxyMaterialUtilities
 		return NumPacked >= 2;
 	}
 
-	static UMaterialInstanceConstant* CreateProxyMaterialInstance(UPackage* InOuter, const FMaterialProxySettings& InMaterialProxySettings, UMaterialInterface* InBaseMaterial, FFlattenMaterial& FlattenMaterial, const FString& AssetBasePath, const FString& AssetBaseName, TArray<UObject*>& OutAssetsToSync)
+	static UMaterialInstanceConstant* CreateProxyMaterialInstance(UPackage* InOuter, const FMaterialProxySettings& InMaterialProxySettings, UMaterialInterface* InBaseMaterial, const FFlattenMaterial& FlattenMaterial, const FString& AssetBasePath, const FString& AssetBaseName, TArray<UObject*>& OutAssetsToSync)
 	{
 		UMaterialInterface* BaseMaterial = InBaseMaterial;
 		
@@ -136,6 +136,19 @@ namespace ProxyMaterialUtilities
 
 		FStaticParameterSet NewStaticParameterSet;
 
+		if(FlattenMaterial.UVChannel != 0)
+		{
+			// If the used texture coordinate was not the default UV0 set the appropriate one on the instance material
+			FStaticSwitchParameter SwitchParameter;
+			SwitchParameter.ParameterInfo.Name = TEXT("UseCustomUV");
+			SwitchParameter.Value = true;
+			SwitchParameter.bOverride = true;
+			NewStaticParameterSet.StaticSwitchParameters.Add(SwitchParameter);
+
+			SwitchParameter.ParameterInfo.Name = *(TEXT("UseUV") + FString::FromInt(FlattenMaterial.UVChannel));
+			NewStaticParameterSet.StaticSwitchParameters.Add(SwitchParameter);
+		}
+
 		// Load textures and set switches accordingly
 		if (FlattenMaterial.GetPropertySize(EFlattenMaterialProperties::Diffuse).Num() > 0 && !(FlattenMaterial.IsPropertyConstant(EFlattenMaterialProperties::Diffuse) && FlattenMaterial.GetPropertySamples(EFlattenMaterialProperties::Diffuse)[0] == FColor::Black))
 		{
diff --git a/Engine/Source/Runtime/Engine/Classes/Engine/MeshMerging.h b/Engine/Source/Runtime/Engine/Classes/Engine/MeshMerging.h
--- a/Engine/Source/Runtime/Engine/Classes/Engine/MeshMerging.h
+++ b/Engine/Source/Runtime/Engine/Classes/Engine/MeshMerging.h
@@ -470,6 +470,10 @@ struct FMeshMergingSettings
 	UPROPERTY(EditAnywhere, BlueprintReadWrite, Category = MaterialSettings)
 	uint8 bMergeMaterials:1;
 
+	/** Create a flat material from all source materials, along with a new set of UVs. This material won't be applied to any section by default. */
+	UPROPERTY(EditAnywhere, BlueprintReadWrite, Category = MaterialSettings)
+	uint8 bCreateMergedMaterial : 1;
+
 	/** Whether or not vertex data such as vertex colours should be baked into the resulting mesh */
 	UPROPERTY(EditAnywhere, BlueprintReadWrite, Category = MeshSettings)
 	uint8 bBakeVertexDataToMesh:1;
@@ -540,6 +544,7 @@ struct FMeshMergingSettings
 		, bPivotPointAtZero(false)
 		, bMergePhysicsData(false)
 		, bMergeMaterials(false)
+		, bCreateMergedMaterial(false)
 		, bBakeVertexDataToMesh(false)
 		, bUseVertexDataForBakingMaterial(true)
 		, bUseTextureBinning(false)
