GenerateProxyMesh causes incorrect transform in ISM using Scale
jeanfrancois.dube

diff --git a/Engine/Source/Developer/MeshMergeUtilities/Private/MeshMergeHelpers.cpp b/Engine/Source/Developer/MeshMergeUtilities/Private/MeshMergeHelpers.cpp
--- a/Engine/Source/Developer/MeshMergeUtilities/Private/MeshMergeHelpers.cpp
+++ b/Engine/Source/Developer/MeshMergeUtilities/Private/MeshMergeHelpers.cpp
@@ -169,22 +169,16 @@ void FMeshMergeHelpers::ExtractSections(const UStaticMesh* StaticMesh, int32 LOD
 void FMeshMergeHelpers::ExpandInstances(const UInstancedStaticMeshComponent* InInstancedStaticMeshComponent, FMeshDescription& InOutRawMesh, TArray<FSectionInfo>& InOutSections)
 {
 	FMeshDescription CombinedRawMesh;
+	UStaticMesh::RegisterMeshAttributes(CombinedRawMesh);
+
+	FTransform ComponentTransform = InInstancedStaticMeshComponent->GetComponentTransform();
+	FTransform ComponentTransformInv = ComponentTransform.Inverse();
 
-	bool bFirstMesh = true;
 	for(const FInstancedStaticMeshInstanceData& InstanceData : InInstancedStaticMeshComponent->PerInstanceSMData)
 	{
 		FMeshDescription InstanceRawMesh = InOutRawMesh;
-		FMeshMergeHelpers::TransformRawMeshVertexData(FTransform(InstanceData.Transform), InstanceRawMesh);
-
-		if(bFirstMesh)
-		{
-			CombinedRawMesh = InstanceRawMesh;
-			bFirstMesh = false;
-		}
-		else
-		{
-			FMeshMergeHelpers::AppendRawMesh(CombinedRawMesh, InstanceRawMesh);
-		}
+		FMeshMergeHelpers::TransformRawMeshVertexData(ComponentTransformInv * FTransform(InstanceData.Transform) * ComponentTransform, InstanceRawMesh);
+		FMeshMergeHelpers::AppendRawMesh(CombinedRawMesh, InstanceRawMesh);
 	}
 
 	InOutRawMesh = CombinedRawMesh;
