HLOD - Allow merging of HISM to the resulting mesh
Also allow merging of H(ISM) from the mesh merging tool
jeanfrancois.dube

diff --git a/Engine/Source/Developer/MeshMergeUtilities/Private/MeshMergeUtilities.cpp b/Engine/Source/Developer/MeshMergeUtilities/Private/MeshMergeUtilities.cpp
--- a/Engine/Source/Developer/MeshMergeUtilities/Private/MeshMergeUtilities.cpp
+++ b/Engine/Source/Developer/MeshMergeUtilities/Private/MeshMergeUtilities.cpp
@@ -1119,7 +1119,7 @@ void FMeshMergeUtilities::CreateProxyMesh(const TArray<UStaticMeshComponent*>& I
 	// Remove anything non-regular or non-spline static mesh components
 	ComponentsToMerge.RemoveAll([](UStaticMeshComponent* Val)
 	{
-		if (Val->GetClass() != UStaticMeshComponent::StaticClass() && Val->GetClass() != UInstancedStaticMeshComponent::StaticClass() && !Val->IsA(USplineMeshComponent::StaticClass()))
+		if (Val->GetClass() != UStaticMeshComponent::StaticClass() && !Val->IsA<UInstancedStaticMeshComponent>() && !Val->IsA<USplineMeshComponent>())
 		{
 			return true;
 		}
@@ -1224,7 +1224,7 @@ void FMeshMergeUtilities::CreateProxyMesh(const TArray<UStaticMeshComponent*>& I
 			}
 
 			// If the component is an ISMC then we need to duplicate the vertex data
-			if(StaticMeshComponent->GetClass() == UInstancedStaticMeshComponent::StaticClass())
+			if(StaticMeshComponent->IsA<UInstancedStaticMeshComponent>())
 			{
 				const UInstancedStaticMeshComponent* InstancedStaticMeshComponent = Cast<UInstancedStaticMeshComponent>(StaticMeshComponent);
 				FMeshMergeHelpers::ExpandInstances(InstancedStaticMeshComponent, *RawMesh, Sections);
@@ -1720,7 +1720,7 @@ bool RetrieveRawMeshData(FMeshMergeDataTracker& DataTracker
 	RawMesh.Compact(RemapInformation);
 
 	// If the component is an ISMC then we need to duplicate the vertex data
-	if (Component->GetClass() == UInstancedStaticMeshComponent::StaticClass())
+	if (Component->IsA<UInstancedStaticMeshComponent>())
 	{
 		const UInstancedStaticMeshComponent* InstancedStaticMeshComponent = Cast<UInstancedStaticMeshComponent>(Component);
 		FMeshMergeHelpers::ExpandInstances(InstancedStaticMeshComponent, RawMesh, Sections);
diff --git a/Engine/Source/Editor/MergeActors/Private/MeshProxyTool/MeshProxyTool.cpp b/Engine/Source/Editor/MergeActors/Private/MeshProxyTool/MeshProxyTool.cpp
--- a/Engine/Source/Editor/MergeActors/Private/MeshProxyTool/MeshProxyTool.cpp
+++ b/Engine/Source/Editor/MergeActors/Private/MeshProxyTool/MeshProxyTool.cpp
@@ -9,6 +9,7 @@
 #include "Widgets/DeclarativeSyntaxSupport.h"
 #include "GameFramework/Actor.h"
 #include "Components/StaticMeshComponent.h"
+#include "Components/InstancedStaticMeshComponent.h"
 #include "Components/SplineMeshComponent.h"
 #include "Engine/Selection.h"
 #include "Editor.h"
@@ -150,7 +151,12 @@ bool FMeshProxyTool::RunMerge(const FString& PackageName)
 					StaticMeshComponentsToMerge.Add(StaticMeshComponent);
 			}
 		}
-		StaticMeshComponentsToMerge.RemoveAll([](UStaticMeshComponent* Val) { return !((Val->GetClass() == UStaticMeshComponent::StaticClass() && Val->GetStaticMesh()) || Val->IsA(USplineMeshComponent::StaticClass())); });
+		StaticMeshComponentsToMerge.RemoveAll([](UStaticMeshComponent* Val)
+		{
+			bool bValidClass = Val->GetClass() == UStaticMeshComponent::StaticClass() || Val->IsA<UInstancedStaticMeshComponent>() || Val->IsA<USplineMeshComponent>();
+			bool bValidMesh = Val->GetStaticMesh();
+			return !bValidClass || !bValidMesh;
+		});
 		
 		if ( StaticMeshComponentsToMerge.Num())
 		{
