HLOD building crashes when building
Cherrypicked CL8214461 from Dev-Enterprise

diff --git a/Engine/Source/Developer/MeshMergeUtilities/Private/MeshMergeUtilities.cpp b/Engine/Source/Developer/MeshMergeUtilities/Private/MeshMergeUtilities.cpp
--- a/Engine/Source/Developer/MeshMergeUtilities/Private/MeshMergeUtilities.cpp
+++ b/Engine/Source/Developer/MeshMergeUtilities/Private/MeshMergeUtilities.cpp
@@ -1178,8 +1178,8 @@ void FMeshMergeUtilities::CreateProxyMesh(const TArray<UStaticMeshComponent*>& I
 
 	TArray<FMeshDescription*> RawMeshData;
 
-	// LOD index, <original section index, unique section index>
-	TMultiMap<uint32, TPair<uint32, uint32>> UniqueSectionIndexPerLOD;
+	// Mesh index, <original section index, unique section index>
+	TMultiMap<uint32, TPair<uint32, uint32>> MeshSectionToUniqueSection;
 
 	// Unique set of sections in mesh
 	TArray<FSectionInfo> UniqueSections;
@@ -1217,9 +1217,9 @@ void FMeshMergeUtilities::CreateProxyMesh(const TArray<UStaticMeshComponent*>& I
 			for (int32 SectionIndex = 0; SectionIndex < Sections.Num(); ++SectionIndex)
 			{
 				FSectionInfo& Section = Sections[SectionIndex];
-				const int32 UniqueIndex = UniqueSections.AddUnique(Section);
-				UniqueSectionIndexPerLOD.Add(MeshIndex, TPair<uint32, uint32>(UniqueIndex, Section.MaterialIndex));
 
+				const int32 UniqueIndex = UniqueSections.AddUnique(Section);
+				MeshSectionToUniqueSection.Add(MeshIndex, TPair<uint32, uint32>(SectionIndex, UniqueIndex));
 				SectionToMesh.Add(UniqueIndex, MeshIndex);
 			}
 
@@ -1239,6 +1239,7 @@ void FMeshMergeUtilities::CreateProxyMesh(const TArray<UStaticMeshComponent*>& I
 	}
 
 	TArray<UMaterialInterface*> UniqueMaterials;
+	//Unique material index to unique section index
 	TMultiMap<uint32, uint32> SectionToMaterialMap;
 	for (int32 SectionIndex = 0; SectionIndex < UniqueSections.Num(); ++SectionIndex)
 	{
@@ -1266,6 +1267,7 @@ void FMeshMergeUtilities::CreateProxyMesh(const TArray<UStaticMeshComponent*>& I
 	{
 		UMaterialInterface* Material = UniqueMaterials[MaterialIndex];
 
+		//Unique section indices
 		TArray<uint32> SectionIndices;
 		SectionToMaterialMap.MultiFind(MaterialIndex, SectionIndices);
 
@@ -1323,13 +1325,14 @@ void FMeshMergeUtilities::CreateProxyMesh(const TArray<UStaticMeshComponent*>& I
 					MeshSettings.TextureCoordinateBox = FBox2D(MeshSettings.CustomTextureCoordinates);
 
 					// Section index is a unique one so we need to map it to the mesh's equivalent(s)
-					TArray<TPair<uint32, uint32>> UniqueToMeshSectionIndices;
-					UniqueSectionIndexPerLOD.MultiFind(MeshIndex, UniqueToMeshSectionIndices);
-					for (const TPair<uint32, uint32> IndexPair : UniqueToMeshSectionIndices)
+					TArray<TPair<uint32, uint32>> SectionToUniqueSectionIndices;
+					MeshSectionToUniqueSection.MultiFind(MeshIndex, SectionToUniqueSectionIndices);
+					for (const TPair<uint32, uint32> IndexPair : SectionToUniqueSectionIndices)
 					{
-						if (IndexPair.Key == SectionIndex)
+						if (IndexPair.Value == SectionIndex)
 						{
-							MeshSettings.MaterialIndices.Add(IndexPair.Value);
+							MeshSettings.MaterialIndices.Add(IndexPair.Key);
+							OutputMaterialsMap.Add(MeshIndex, TPair<uint32, uint32>(IndexPair.Key, GlobalMeshSettings.Num()));
 						}
 					}
 
@@ -1346,12 +1349,6 @@ void FMeshMergeUtilities::CreateProxyMesh(const TArray<UStaticMeshComponent*>& I
 						}
 					}
 
-					// For each original material index add an entry to the corresponding LOD and bake output index
-					for (int32 Index : MeshSettings.MaterialIndices)
-					{
-						OutputMaterialsMap.Add(MeshIndex, TPair<uint32, uint32>(Index, GlobalMeshSettings.Num()));
-					}
-
 					GlobalMeshSettings.Add(MeshSettings);
 					GlobalMaterialSettings.Add(MaterialSettings);
 				}
@@ -1373,13 +1370,13 @@ void FMeshMergeUtilities::CreateProxyMesh(const TArray<UStaticMeshComponent*>& I
 
 				for (uint32 MeshIndex : MeshIndices)
 				{
-					TArray<TPair<uint32, uint32>> UniqueToMeshSectionIndices;
-					UniqueSectionIndexPerLOD.MultiFind(MeshIndex, UniqueToMeshSectionIndices);
-					for (const TPair<uint32, uint32> IndexPair : UniqueToMeshSectionIndices)
+					TArray<TPair<uint32, uint32>> SectionToUniqueSectionIndices;
+					MeshSectionToUniqueSection.MultiFind(MeshIndex, SectionToUniqueSectionIndices);
+					for (const TPair<uint32, uint32> IndexPair : SectionToUniqueSectionIndices)
 					{
-						if (IndexPair.Key == SectionIndex)
+						if (IndexPair.Value == SectionIndex)
 						{
-							OutputMaterialsMap.Add(MeshIndex, TPair<uint32, uint32>(IndexPair.Value, GlobalMeshSettings.Num()));
+							OutputMaterialsMap.Add(MeshIndex, TPair<uint32, uint32>(IndexPair.Key, GlobalMeshSettings.Num()));
 						}
 					}
 				}
@@ -1438,19 +1435,13 @@ void FMeshMergeUtilities::CreateProxyMesh(const TArray<UStaticMeshComponent*>& I
 
 			TArray<TPair<uint32, uint32>> SectionAndOutputIndices;
 			OutputMaterialsMap.MultiFind(MeshIndex, SectionAndOutputIndices);
-			//Make sure the section index are in the correct order
-			SectionAndOutputIndices.Sort([](const TPair<uint32, uint32>& A, const TPair<uint32, uint32>& B) { return (A.Key < B.Key); });
-
 			TArray<int32> Remap;
-			TArray<int32> UniqueMaterialIndexes;
 			// Reorder loops 
 			for (const TPair<uint32, uint32>& IndexPair : SectionAndOutputIndices)
 			{
-				//We are not using the IndexPair.Key since we want to keep the polygon group
-				//We instead find the section index by looking at unique IndexPair.Value
+				const int32 SectionIndex = IndexPair.Key;
 				const int32 NewIndex = IndexPair.Value;
 
-				const int32 SectionIndex = UniqueMaterialIndexes.AddUnique(NewIndex);
 				if (Remap.Num() < (SectionIndex + 1))
 				{
 					Remap.SetNum(SectionIndex + 1);
