Fixed merged meshes having wrong material mapping

diff --git a/Engine/Source/Developer/MeshMergeUtilities/Private/MeshMergeUtilities.cpp b/Engine/Source/Developer/MeshMergeUtilities/Private/MeshMergeUtilities.cpp
--- a/Engine/Source/Developer/MeshMergeUtilities/Private/MeshMergeUtilities.cpp
+++ b/Engine/Source/Developer/MeshMergeUtilities/Private/MeshMergeUtilities.cpp
@@ -1176,7 +1176,8 @@ void FMeshMergeUtilities::CreateProxyMesh(const TArray<UStaticMeshComponent*>& I
 	// Mesh / LOD index	
 	TMap<uint32, FMeshDescription*> RawMeshLODs;
 
-	TArray<FMeshDescription*> RawMeshData;
+	// Copies of mesh data
+	TArray<FMeshDescription*> MeshDescriptionData;
 
 	// Mesh index, <original section index, unique section index>
 	TMultiMap<uint32, TPair<uint32, uint32>> MeshSectionToUniqueSection;
@@ -1204,10 +1205,10 @@ void FMeshMergeUtilities::CreateProxyMesh(const TArray<UStaticMeshComponent*>& I
 			static const bool bPropagateVertexColours = true;
 
 			// Retrieve mesh data in FMeshDescription form
-			FMeshDescription* RawMesh = new FMeshDescription();
-			UStaticMesh::RegisterMeshAttributes(*RawMesh);
-			FMeshMergeHelpers::RetrieveMesh(StaticMeshComponent, LODIndex, *RawMesh, bPropagateVertexColours);
-			const int32 MeshIndex = RawMeshData.Add(RawMesh);
+			const int32 MeshIndex = MeshDescriptionData.Add(new FMeshDescription());
+			FMeshDescription& MeshDescription = *MeshDescriptionData[MeshIndex];
+			UStaticMesh::RegisterMeshAttributes(MeshDescription);
+			FMeshMergeHelpers::RetrieveMesh(StaticMeshComponent, LODIndex, MeshDescription, bPropagateVertexColours);
 
 			// Reset section array for reuse
 			Sections.SetNum(0, false);
@@ -1227,7 +1228,7 @@ void FMeshMergeUtilities::CreateProxyMesh(const TArray<UStaticMeshComponent*>& I
 			if(StaticMeshComponent->IsA<UInstancedStaticMeshComponent>())
 			{
 				const UInstancedStaticMeshComponent* InstancedStaticMeshComponent = Cast<UInstancedStaticMeshComponent>(StaticMeshComponent);
-				FMeshMergeHelpers::ExpandInstances(InstancedStaticMeshComponent, *RawMesh, Sections);
+				FMeshMergeHelpers::ExpandInstances(InstancedStaticMeshComponent, MeshDescription, Sections);
 				NumInstances = InstancedStaticMeshComponent->PerInstanceSMData.Num();
 			}
 		}	
@@ -1297,8 +1298,8 @@ void FMeshMergeUtilities::CreateProxyMesh(const TArray<UStaticMeshComponent*>& I
 				for (uint32 MeshIndex : MeshIndices)
 				{
 					FMeshData MeshSettings;
-					// Retrieve raw mesh
-					FMeshDescription* MeshDescription = RawMeshData[MeshIndex];
+					// Retrieve mesh description
+					FMeshDescription* MeshDescription = MeshDescriptionData[MeshIndex];
 					MeshSettings.RawMeshDescription = MeshDescription;
 
 					TVertexInstanceAttributesRef<FVector2D> VertexInstanceUVs = MeshSettings.RawMeshDescription->VertexInstanceAttributes().GetAttributesRef<FVector2D>(MeshAttribute::VertexInstance::TextureCoordinate);
@@ -1391,7 +1392,7 @@ void FMeshMergeUtilities::CreateProxyMesh(const TArray<UStaticMeshComponent*>& I
 	IMaterialBakingModule& MaterialBakingModule = FModuleManager::Get().LoadModuleChecked<IMaterialBakingModule>("MaterialBaking");
 
 	auto MaterialFlattenLambda =
-		[this, &Options, &GlobalMeshSettings, &GlobalMaterialSettings, &RawMeshData, &OutputMaterialsMap, &MaterialBakingModule](TArray<FFlattenMaterial>& FlattenedMaterialArray)
+		[this, &Options, &GlobalMeshSettings, &GlobalMaterialSettings, &MeshDescriptionData, &OutputMaterialsMap, &MaterialBakingModule](TArray<FFlattenMaterial>& FlattenedMaterialArray)
 	{
 		TArray<FMeshData*> MeshSettingPtrs;
 		for (int32 SettingsIndex = 0; SettingsIndex < GlobalMeshSettings.Num(); ++SettingsIndex)
@@ -1429,9 +1430,9 @@ void FMeshMergeUtilities::CreateProxyMesh(const TArray<UStaticMeshComponent*>& I
 		ConvertOutputToFlatMaterials(BakeOutputs, GlobalMaterialSettings, FlattenedMaterialArray);
 
 		// Now have the baked out material data, need to have a map or actually remap the raw mesh data to baked material indices
-		for (int32 MeshIndex = 0; MeshIndex < RawMeshData.Num(); ++MeshIndex)
+		for (int32 MeshIndex = 0; MeshIndex < MeshDescriptionData.Num(); ++MeshIndex)
 		{
-			FMeshDescription& RawMesh = *RawMeshData[MeshIndex];
+			FMeshDescription& MeshDescription = *MeshDescriptionData[MeshIndex];
 
 			TArray<TPair<uint32, uint32>> SectionAndOutputIndices;
 			OutputMaterialsMap.MultiFind(MeshIndex, SectionAndOutputIndices);
@@ -1451,13 +1452,13 @@ void FMeshMergeUtilities::CreateProxyMesh(const TArray<UStaticMeshComponent*>& I
 			}
 			
 			TMap<FPolygonGroupID, FPolygonGroupID> RemapPolygonGroup;
-			for (const FPolygonGroupID& PolygonGroupID : RawMesh.PolygonGroups().GetElementIDs())
+			for (const FPolygonGroupID& PolygonGroupID : MeshDescription.PolygonGroups().GetElementIDs())
 			{
 				checkf(Remap.IsValidIndex(PolygonGroupID.GetValue()), TEXT("Missing material bake output index entry for mesh(section)"));
 				int32 RemapID = Remap[PolygonGroupID.GetValue()];
 				RemapPolygonGroup.Add(PolygonGroupID, FPolygonGroupID(RemapID));
 			}
-			FMeshDescriptionOperations::RemapPolygonGroups(RawMesh, RemapPolygonGroup);
+			FMeshDescriptionOperations::RemapPolygonGroups(MeshDescription, RemapPolygonGroup);
 		}
 	};
 
@@ -1492,11 +1493,11 @@ void FMeshMergeUtilities::CreateProxyMesh(const TArray<UStaticMeshComponent*>& I
 
 	// We are only using LOD level 0 (ProxyMeshTargetLODLevel)
 	TArray<FMeshMergeData> MergeDataEntries;
-	for (int32 Index = 0; Index < RawMeshData.Num(); ++Index)
+	for (int32 Index = 0; Index < MeshDescriptionData.Num(); ++Index)
 	{
 		FMeshMergeData MergeData;
 		MergeData.SourceStaticMesh = ComponentsToMerge[Index]->GetStaticMesh();
-		MergeData.RawMesh = RawMeshData[Index];
+		MergeData.RawMesh = MeshDescriptionData[Index];
 		MergeData.bIsClippingMesh = false;
 
 		FMeshMergeHelpers::CalculateTextureCoordinateBoundsForRawMesh(*MergeData.RawMesh, MergeData.TexCoordBounds);
@@ -1570,6 +1571,12 @@ void FMeshMergeUtilities::CreateProxyMesh(const TArray<UStaticMeshComponent*>& I
 	{
 		delete RawMesh;
 	}
+
+	// Clean up the MeshDescriptionData
+	for (FMeshDescription* MeshDescription : MeshDescriptionData)
+	{
+		delete MeshDescription;
+	}
 }
 
 
@@ -1680,7 +1687,6 @@ bool RetrieveRawMeshData(FMeshMergeDataTracker& DataTracker
 	// Extract sections for given LOD index from the mesh 
 	Adapter.RetrieveMeshSections(LODIndex, Sections);
 
-	TMap<FPolygonGroupID, FPolygonGroupID> RemapPolygonGroup;
 	for (int32 SectionIndex = 0; SectionIndex < Sections.Num(); ++SectionIndex)
 	{
 		const FSectionInfo& Section = Sections[SectionIndex];
@@ -1688,22 +1694,13 @@ bool RetrieveRawMeshData(FMeshMergeDataTracker& DataTracker
 		const int32 UniqueIndex = DataTracker.AddSection(Section);
 
 		// Store of original to unique section index entry for this component + LOD index
-		DataTracker.AddSectionRemapping(ComponentIndex, LODIndex, Section.MaterialIndex, UniqueIndex);
+		DataTracker.AddSectionRemapping(ComponentIndex, LODIndex, SectionIndex, UniqueIndex);
 		DataTracker.AddMaterialSlotName(Section.Material, Section.MaterialSlotName);
 
 		if (!bMergeMaterialData)
 		{
 			FMeshDescriptionOperations::SwapPolygonPolygonGroup(RawMesh, UniqueIndex, Section.StartIndex, Section.EndIndex, false);
 		}
-		else
-		{
-			RemapPolygonGroup.Add(FPolygonGroupID(SectionIndex), FPolygonGroupID(Section.MaterialIndex));
-		}
-	}
-	
-	if (bMergeMaterialData)
-	{
-		FMeshDescriptionOperations::RemapPolygonGroups(RawMesh, RemapPolygonGroup);
 	}
 	
 	//Compact the PolygonGroupID to make sure it follow the section index
@@ -1901,7 +1898,6 @@ void FMeshMergeUtilities::MergeComponentsToStaticMesh(const TArray<UPrimitiveCom
 	// Find all unique materials and remap section to unique materials
 	TArray<UMaterialInterface*> UniqueMaterials;
 	TMap<UMaterialInterface*, int32> MaterialIndices;
-	TMultiMap<uint32, uint32> SectionToMaterialMap;
 	TMap<UMaterialInterface*, UMaterialInterface*> CollapsedMaterialMap;
 
 	for (int32 SectionIndex = 0; SectionIndex < DataTracker.NumberOfUniqueSections(); ++SectionIndex)
@@ -1925,9 +1921,6 @@ void FMeshMergeUtilities::MergeComponentsToStaticMesh(const TArray<UPrimitiveCom
 
 		// Update map to 'collapsed' materials
 		CollapsedMaterialMap.Add(MaterialInterface, UniqueMaterials[UniqueIndex]);
-
-		// Store off usage of unique material by unique sections
-		SectionToMaterialMap.Add(UniqueIndex, SectionIndex);
 	}
 
 	// For each unique material calculate how 'important' they are
@@ -1963,6 +1956,10 @@ void FMeshMergeUtilities::MergeComponentsToStaticMesh(const TArray<UPrimitiveCom
 
 		TMap<UMaterialInterface*, int32> MaterialToDefaultMeshData;
 
+		// If we are generating a single LOD and want to merge materials we can utilize texture space better by generating unique UVs
+		// for the merged mesh and baking out materials using those UVs
+		const bool bGloballyRemapUVs = !bMergeAllLODs && !InSettings.bReuseMeshLightmapUVs;
+
 		for (TConstRawMeshIterator RawMeshIterator = DataTracker.GetConstRawMeshIterator(); RawMeshIterator; ++RawMeshIterator)
 		{
 			const FMeshLODKey& Key = RawMeshIterator.Key();
@@ -2005,7 +2002,7 @@ void FMeshMergeUtilities::MergeComponentsToStaticMesh(const TArray<UPrimitiveCom
 				MeshData.bMirrored = Component->GetComponentTransform().GetDeterminant() < 0.0f;
 				int32 MeshDataIndex = 0;
 
-				if (InSettings.bCreateMergedMaterial || (InSettings.bUseVertexDataForBakingMaterial && (bDoesMaterialUseVertexData || bRequiresUniqueUVs)))
+				if (InSettings.bCreateMergedMaterial || bGloballyRemapUVs || (InSettings.bUseVertexDataForBakingMaterial && (bDoesMaterialUseVertexData || bRequiresUniqueUVs)))
 				{
 					MeshData.RawMeshDescription = DataTracker.GetRawMeshPtr(Key);
 
@@ -2119,9 +2116,6 @@ void FMeshMergeUtilities::MergeComponentsToStaticMesh(const TArray<UPrimitiveCom
 			MaterialSettingPtrs.Add(&GlobalMaterialSettings[SettingsIndex]);
 		}
 
-		// If we are generating a single LOD and want to merge materials we can utilize texture space better by generating unique UVs
-		// for the merged mesh and baking out materials using those UVs
-		const bool bGloballyRemapUVs = !bMergeAllLODs && !InSettings.bReuseMeshLightmapUVs;
 		if(bGloballyRemapUVs)
 		{
 			TArray<FMeshDescription> MergedRawMeshes;
