Do not exclude static mesh subclasses from the mesh merge process

diff --git a/Engine/Source/Developer/MeshMergeUtilities/Private/MeshMergeUtilities.cpp b/Engine/Source/Developer/MeshMergeUtilities/Private/MeshMergeUtilities.cpp
--- a/Engine/Source/Developer/MeshMergeUtilities/Private/MeshMergeUtilities.cpp
+++ b/Engine/Source/Developer/MeshMergeUtilities/Private/MeshMergeUtilities.cpp
@@ -1116,21 +1116,8 @@ void FMeshMergeUtilities::CreateProxyMesh(const TArray<UStaticMeshComponent*>& I
 
 	TArray<UStaticMeshComponent*> ComponentsToMerge = InComponentsToMerge;
 
-	// Remove anything non-regular or non-spline static mesh components
-	ComponentsToMerge.RemoveAll([](UStaticMeshComponent* Val)
-	{
-		if (Val->GetClass() != UStaticMeshComponent::StaticClass() && !Val->IsA<UInstancedStaticMeshComponent>() && !Val->IsA<USplineMeshComponent>())
-		{
-			return true;
-		}
-
-		if (Val->GetStaticMesh() == nullptr)
-		{
-			return true;
-		}
-
-		return false;
-	});
+	// Remove invalid components
+	ComponentsToMerge.RemoveAll([](UStaticMeshComponent* Val) { return Val->GetStaticMesh() == nullptr; });
 
 	// No actors given as input
 	if (ComponentsToMerge.Num() == 0)
diff --git a/Engine/Source/Editor/MergeActors/Private/MeshProxyTool/MeshProxyTool.cpp b/Engine/Source/Editor/MergeActors/Private/MeshProxyTool/MeshProxyTool.cpp
--- a/Engine/Source/Editor/MergeActors/Private/MeshProxyTool/MeshProxyTool.cpp
+++ b/Engine/Source/Editor/MergeActors/Private/MeshProxyTool/MeshProxyTool.cpp
@@ -151,12 +151,7 @@ bool FMeshProxyTool::RunMerge(const FString& PackageName)
 					StaticMeshComponentsToMerge.Add(StaticMeshComponent);
 			}
 		}
-		StaticMeshComponentsToMerge.RemoveAll([](UStaticMeshComponent* Val)
-		{
-			bool bValidClass = Val->GetClass() == UStaticMeshComponent::StaticClass() || Val->IsA<UInstancedStaticMeshComponent>() || Val->IsA<USplineMeshComponent>();
-			bool bValidMesh = Val->GetStaticMesh() != nullptr;
-			return !bValidClass || !bValidMesh;
-		});
+		StaticMeshComponentsToMerge.RemoveAll([](UStaticMeshComponent* Val) { return Val->GetStaticMesh() == nullptr; });
 		
 		if ( StaticMeshComponentsToMerge.Num())
 		{
