HLOD useLODAImposter seems to not copy over all UV channels
Mesh merging code was assuming the UV2 & UV3 channels were free to use to store position & scale of imposter instances.
This behavior seems to have been hacked in the engine for an internal project & is not needed anymore since it's now possible to use instanced static meshes with LODActors
Still kept old behavior around (CVar) in case someone would really need it

diff --git a/Engine/Source/Developer/MeshMergeUtilities/Private/MeshMergeHelpers.cpp b/Engine/Source/Developer/MeshMergeUtilities/Private/MeshMergeHelpers.cpp
--- a/Engine/Source/Developer/MeshMergeUtilities/Private/MeshMergeHelpers.cpp
+++ b/Engine/Source/Developer/MeshMergeUtilities/Private/MeshMergeHelpers.cpp
@@ -46,6 +46,14 @@
 
 //DECLARE_LOG_CATEGORY_CLASS(LogMeshMerging, Verbose, All);
 
+static TAutoConsoleVariable<int32> CVarMeshMergeStoreImposterInfoInUVs(
+	TEXT("r.MeshMerge.StoreImposterInfoInUVs"),
+	0,
+	TEXT("Determines whether or not to store imposter info (position.xy in UV2, position.z + scale in UV3) in the merged mesh UV channels\n")
+	TEXT("0: Do not store imposters info in UVs (default)\n")
+	TEXT("1: Store imposter info in UVs (legacy)\n"),
+	ECVF_Default);
+
 void FMeshMergeHelpers::ExtractSections(const UStaticMeshComponent* Component, int32 LODIndex, TArray<FSectionInfo>& OutSections)
 {
 	static UMaterialInterface* DefaultMaterial = UMaterial::GetDefaultMaterial(MD_Surface);
@@ -1294,66 +1302,42 @@ void FMeshMergeHelpers::ExtractImposterToRawMesh(const UStaticMeshComponent* InI
 
 void FMeshMergeHelpers::MergeImpostersToRawMesh(TArray<const UStaticMeshComponent*> ImposterComponents, FMeshDescription& InRawMesh, const FVector& InPivot, int32 InBaseMaterialIndex, TArray<UMaterialInterface*>& OutImposterMaterials)
 {
-	// TODO decide whether we want this to be user specified or derived from the RawMesh
-	/*const int32 UVOneIndex = [RawMesh, Data]() -> int32
+	TMap<UMaterialInterface*, FPolygonGroupID> ImposterMaterialToPolygonGroupID;
+	for (const UStaticMeshComponent* Component : ImposterComponents)
 	{
-		int32 ChannelIndex = 0;
-		for (; ChannelIndex < MAX_MESH_TEXTURE_COORDS; ++ChannelIndex)
-		{
-			if (RawMesh.WedgeTexCoords[ChannelIndex].Num() == 0)
-			{
-				break;
-			}
-		}
+		// Retrieve imposter LOD mesh and material			
+		const int32 LODIndex = Component->GetStaticMesh()->GetNumLODs() - 1;
 
-		int32 MaxUVChannel = ChannelIndex;
-		for (const UStaticMeshComponent* Component : ImposterComponents)
-		{
-			MaxUVChannel = FMath::Max(MaxUVChannel, Component->GetStaticMesh()->RenderData->LODResources[Component->GetStaticMesh()->GetNumLODs() - 1].GetNumTexCoords());
-		}
-
-		return MaxUVChannel;
-	}();*/
+		// Retrieve mesh data in FMeshDescription form
+		FMeshDescription ImposterMesh;
+		UStaticMesh::RegisterMeshAttributes(ImposterMesh);
+		FMeshMergeHelpers::RetrieveMesh(Component, LODIndex, ImposterMesh, false);
 
-	const int32 UVOneIndex = 2; // if this is changed back to being dynamic, renable the if statement below
+		// Retrieve the sections, we're expect 1 for imposter meshes
+		TArray<FSectionInfo> Sections;
+		FMeshMergeHelpers::ExtractSections(Component, LODIndex, Sections);
 
-	// Ensure there are enough UV channels available to store the imposter data
-	//if (UVOneIndex != INDEX_NONE && UVOneIndex < (MAX_MESH_TEXTURE_COORDS - 2))
-	{
-		TMap<UMaterialInterface*, FPolygonGroupID> ImposterMaterialToPolygonGroupID;
-		for (const UStaticMeshComponent* Component : ImposterComponents)
+		TArray<int32> SectionImposterUniqueMaterialIndex;
+		for (FSectionInfo& Info : Sections)
 		{
-			// Retrieve imposter LOD mesh and material			
-			const int32 LODIndex = Component->GetStaticMesh()->GetNumLODs() - 1;
-
-			// Retrieve mesh data in FMeshDescription form
-			FMeshDescription ImposterMesh;
-			FStaticMeshAttributes ImposterMeshAttributes(ImposterMesh);
-			ImposterMeshAttributes.Register();
-			FMeshMergeHelpers::RetrieveMesh(Component, LODIndex, ImposterMesh, false);
-
-			// Retrieve the sections, we're expect 1 for imposter meshes
-			TArray<FSectionInfo> Sections;
-			FMeshMergeHelpers::ExtractSections(Component, LODIndex, Sections);
-
-			TArray<int32> SectionImposterUniqueMaterialIndex;
-			for (FSectionInfo& Info : Sections)
-			{
-				SectionImposterUniqueMaterialIndex.Add(OutImposterMaterials.AddUnique(Info.Material));
-			}
+			SectionImposterUniqueMaterialIndex.Add(OutImposterMaterials.AddUnique(Info.Material));
+		}
 
+		if (CVarMeshMergeStoreImposterInfoInUVs.GetValueOnAnyThread())
+		{
 			// Imposter magic, we're storing the actor world position and X scale spread across two UV channels
+			const int32 UVOneIndex = 2;
 			const int32 UVTwoIndex = UVOneIndex + 1;
 			TVertexInstanceAttributesRef<FVector2D> VertexInstanceUVs = ImposterMeshAttributes.GetVertexInstanceUVs();
 			VertexInstanceUVs.SetNumIndices(UVTwoIndex + 1);
 			const int32 NumIndices = ImposterMesh.VertexInstances().Num();
 			const FTransform& ActorToWorld = Component->GetOwner()->GetActorTransform();
 			const FVector ActorPosition = ActorToWorld.TransformPosition(FVector::ZeroVector) - InPivot;
-			for(const FVertexInstanceID& VertexInstanceID : ImposterMesh.VertexInstances().GetElementIDs())
+			for (const FVertexInstanceID& VertexInstanceID : ImposterMesh.VertexInstances().GetElementIDs())
 			{
 				FVector2D UVOne;
 				FVector2D UVTwo;
-					
+
 				UVOne.X = ActorPosition.X;
 				UVOne.Y = ActorPosition.Y;
 				VertexInstanceUVs.Set(VertexInstanceID, UVOneIndex, UVOne);
@@ -1362,32 +1346,32 @@ void FMeshMergeHelpers::MergeImpostersToRawMesh(TArray<const UStaticMeshComponen
 				UVTwo.Y = FMath::Abs(ActorToWorld.GetScale3D().X);
 				VertexInstanceUVs.Set(VertexInstanceID, UVTwoIndex, UVTwo);
 			}
+		}
 
-			TPolygonGroupAttributesRef<FName> SourcePolygonGroupImportedMaterialSlotNames = ImposterMeshAttributes.GetPolygonGroupMaterialSlotNames();
-			TPolygonGroupAttributesRef<FName> TargetPolygonGroupImportedMaterialSlotNames = InRawMesh.PolygonGroupAttributes().GetAttributesRef<FName>(MeshAttribute::PolygonGroup::ImportedMaterialSlotName);
+		TPolygonGroupAttributesRef<FName> SourcePolygonGroupImportedMaterialSlotNames = ImposterMesh.PolygonGroupAttributes().GetAttributesRef<FName>(MeshAttribute::PolygonGroup::ImportedMaterialSlotName);
+		TPolygonGroupAttributesRef<FName> TargetPolygonGroupImportedMaterialSlotNames = InRawMesh.PolygonGroupAttributes().GetAttributesRef<FName>(MeshAttribute::PolygonGroup::ImportedMaterialSlotName);
 
-			//Add the missing polygon group ID to the target(InRawMesh)
-			//Remap the source mesh(ImposterMesh) polygongroup to fit with the target polygon groups
-			TMap<FPolygonGroupID, FPolygonGroupID> RemapSourcePolygonGroup;
-			RemapSourcePolygonGroup.Reserve(ImposterMesh.PolygonGroups().Num());
-			int32 SectionIndex = 0;
-			for (const FPolygonGroupID& SourcePolygonGroupID : ImposterMesh.PolygonGroups().GetElementIDs())
+		//Add the missing polygon group ID to the target(InRawMesh)
+		//Remap the source mesh(ImposterMesh) polygongroup to fit with the target polygon groups
+		TMap<FPolygonGroupID, FPolygonGroupID> RemapSourcePolygonGroup;
+		RemapSourcePolygonGroup.Reserve(ImposterMesh.PolygonGroups().Num());
+		int32 SectionIndex = 0;
+		for (const FPolygonGroupID& SourcePolygonGroupID : ImposterMesh.PolygonGroups().GetElementIDs())
+		{
+			UMaterialInterface* MaterialUseBySection = OutImposterMaterials[SectionImposterUniqueMaterialIndex[SectionIndex++]];
+			FPolygonGroupID* ExistTargetPolygonGroupID = ImposterMaterialToPolygonGroupID.Find(MaterialUseBySection);
+			FPolygonGroupID MatchTargetPolygonGroupID = ExistTargetPolygonGroupID == nullptr ? FPolygonGroupID::Invalid : *ExistTargetPolygonGroupID;
+			if (MatchTargetPolygonGroupID == FPolygonGroupID::Invalid)
 			{
-				UMaterialInterface* MaterialUseBySection = OutImposterMaterials[SectionImposterUniqueMaterialIndex[SectionIndex++]];
-				FPolygonGroupID* ExistTargetPolygonGroupID = ImposterMaterialToPolygonGroupID.Find(MaterialUseBySection);
-				FPolygonGroupID MatchTargetPolygonGroupID = ExistTargetPolygonGroupID == nullptr ? FPolygonGroupID::Invalid : *ExistTargetPolygonGroupID;
-				if (MatchTargetPolygonGroupID == FPolygonGroupID::Invalid)
-				{
-					MatchTargetPolygonGroupID = InRawMesh.CreatePolygonGroup();
-					//use the material name to fill the imported material name. Material name will be unique
-					TargetPolygonGroupImportedMaterialSlotNames[MatchTargetPolygonGroupID] = MaterialUseBySection->GetFName();
-					ImposterMaterialToPolygonGroupID.Add(MaterialUseBySection, MatchTargetPolygonGroupID);
-				}
-				RemapSourcePolygonGroup.Add(SourcePolygonGroupID, MatchTargetPolygonGroupID);
+				MatchTargetPolygonGroupID = InRawMesh.CreatePolygonGroup();
+				//use the material name to fill the imported material name. Material name will be unique
+				TargetPolygonGroupImportedMaterialSlotNames[MatchTargetPolygonGroupID] = MaterialUseBySection->GetFName();
+				ImposterMaterialToPolygonGroupID.Add(MaterialUseBySection, MatchTargetPolygonGroupID);
 			}
-			FMeshDescriptionOperations::RemapPolygonGroups(ImposterMesh, RemapSourcePolygonGroup);
-
-			FMeshMergeHelpers::AppendRawMesh(InRawMesh, ImposterMesh);
+			RemapSourcePolygonGroup.Add(SourcePolygonGroupID, MatchTargetPolygonGroupID);
 		}
+		FMeshDescriptionOperations::RemapPolygonGroups(ImposterMesh, RemapSourcePolygonGroup);
+
+		FMeshMergeHelpers::AppendRawMesh(InRawMesh, ImposterMesh);
 	}
 }
