Fix skeletalmesh bake material crash, and incorrect ID remapping.

diff --git a/Engine/Source/Developer/MeshMergeUtilities/Private/MeshMergeHelpers.cpp b/Engine/Source/Developer/MeshMergeUtilities/Private/MeshMergeHelpers.cpp
--- a/Engine/Source/Developer/MeshMergeUtilities/Private/MeshMergeHelpers.cpp
+++ b/Engine/Source/Developer/MeshMergeUtilities/Private/MeshMergeHelpers.cpp
@@ -109,17 +109,26 @@ void FMeshMergeHelpers::ExtractSections(const USkeletalMeshComponent* Component,
 
 	TArray<FName> MaterialSlotNames = Component->GetMaterialSlotNames();
 
+	const FSkeletalMeshLODInfo* LODInfoPtr = Component->SkeletalMesh->GetLODInfo(LODIndex);
+	check(LODInfoPtr);
 	const FSkeletalMeshLODModel& Model = Resource->LODModels[LODIndex];
-	for (const FSkelMeshSection& MeshSection : Model.Sections)
+	for (int32 SectionIndex = 0; SectionIndex < Model.Sections.Num(); ++SectionIndex)
 	{
+		const FSkelMeshSection& MeshSection = Model.Sections[SectionIndex];
 		// Retrieve material for this section
-		UMaterialInterface* StoredMaterial = Component->GetMaterial(MeshSection.MaterialIndex);
+		int32 MaterialIndex = MeshSection.MaterialIndex;
+		if (LODInfoPtr->LODMaterialMap.IsValidIndex(SectionIndex) && LODInfoPtr->LODMaterialMap[SectionIndex] != INDEX_NONE)
+		{
+			MaterialIndex = LODInfoPtr->LODMaterialMap[SectionIndex];
+		}
+		
+		UMaterialInterface* StoredMaterial = Component->GetMaterial(MaterialIndex);
 		// Make sure the resource actual exists, otherwise use default material
 		StoredMaterial = (StoredMaterial != nullptr) && StoredMaterial->GetMaterialResource(GMaxRHIFeatureLevel) ? StoredMaterial : DefaultMaterial;
 
 		FSectionInfo SectionInfo;
 		SectionInfo.Material = StoredMaterial;
-		SectionInfo.MaterialSlotName = MaterialSlotNames.IsValidIndex(MeshSection.MaterialIndex) ? MaterialSlotNames[MeshSection.MaterialIndex] : NAME_None;
+		SectionInfo.MaterialSlotName = MaterialSlotNames.IsValidIndex(MaterialIndex) ? MaterialSlotNames[MaterialIndex] : NAME_None;
 
 		if (MeshSection.bCastShadow && Component->CastShadow)
 		{
@@ -309,16 +318,15 @@ void FMeshMergeHelpers::RetrieveMesh(USkeletalMeshComponent* SkeletalMeshCompone
 			const int32 NumWedges = SkelMeshSection.NumTriangles * 3;
 
 			//Create the polygon group ID
-			int32 SectionMaterialIndex = SkelMeshSection.MaterialIndex;
-			int32 MaterialIndex = SectionMaterialIndex;
-			// use the remapping of material indices for all LODs besides the base LOD 
-			if (LODIndex > 0 && SrcLODInfo.LODMaterialMap.IsValidIndex(SkelMeshSection.MaterialIndex))
+			int32 MaterialIndex = SkelMeshSection.MaterialIndex;
+			// use the remapping of material indices if there is a valid value
+			if (SrcLODInfo.LODMaterialMap.IsValidIndex(SectionIndex) && SrcLODInfo.LODMaterialMap[SectionIndex] != INDEX_NONE)
 			{
-				MaterialIndex = FMath::Clamp<int32>(SrcLODInfo.LODMaterialMap[SkelMeshSection.MaterialIndex], 0, SkeletalMeshComponent->SkeletalMesh->Materials.Num());
+				MaterialIndex = FMath::Clamp<int32>(SrcLODInfo.LODMaterialMap[SectionIndex], 0, SkeletalMeshComponent->SkeletalMesh->Materials.Num());
 			}
 
 			FName ImportedMaterialSlotName = SkeletalMeshComponent->SkeletalMesh->Materials[MaterialIndex].ImportedMaterialSlotName;
-			const FPolygonGroupID SectionPolygonGroupID(SectionMaterialIndex);
+			const FPolygonGroupID SectionPolygonGroupID(SectionIndex);
 			if (!RawMesh.IsPolygonGroupValid(SectionPolygonGroupID))
 			{
 				RawMesh.CreatePolygonGroupWithID(SectionPolygonGroupID);
diff --git a/Engine/Source/Developer/MeshMergeUtilities/Private/MeshMergeUtilities.cpp b/Engine/Source/Developer/MeshMergeUtilities/Private/MeshMergeUtilities.cpp
--- a/Engine/Source/Developer/MeshMergeUtilities/Private/MeshMergeUtilities.cpp
+++ b/Engine/Source/Developer/MeshMergeUtilities/Private/MeshMergeUtilities.cpp
@@ -417,6 +417,7 @@ void FMeshMergeUtilities::BakeMaterialsForComponent(TArray<TWeakObjectPtr<UObjec
 			int32 SetIndex = Adapter->GetMaterialIndex(LODIndex, Pair.Key);
 			if (!NonReplaceMaterialIndices.Contains(SetIndex))
 			{
+				//TODO (Bug), need to pass the material data MaterialSlotName and ImportedMaterialSlotName. We loose all this data when baking material on skeletalmesh
 				Adapter->SetMaterial(SetIndex, NewMaterials[Pair.Value]);
 			}
 			else
diff --git a/Engine/Source/Developer/MeshMergeUtilities/Private/SkeletalMeshAdapter.cpp b/Engine/Source/Developer/MeshMergeUtilities/Private/SkeletalMeshAdapter.cpp
--- a/Engine/Source/Developer/MeshMergeUtilities/Private/SkeletalMeshAdapter.cpp
+++ b/Engine/Source/Developer/MeshMergeUtilities/Private/SkeletalMeshAdapter.cpp
@@ -31,6 +31,11 @@ void FSkeletalMeshComponentAdapter::RetrieveMeshSections(int32 LODIndex, TArray<
 
 int32 FSkeletalMeshComponentAdapter::GetMaterialIndex(int32 LODIndex, int32 SectionIndex) const
 {
+	const FSkeletalMeshLODInfo* LODInfoPtr = SkeletalMesh->GetLODInfo(LODIndex);
+	if (LODInfoPtr && LODInfoPtr->LODMaterialMap.IsValidIndex(SectionIndex) && LODInfoPtr->LODMaterialMap[SectionIndex] != INDEX_NONE)
+	{
+		return LODInfoPtr->LODMaterialMap[SectionIndex];
+	}
 	return SkeletalMesh->GetImportedModel()->LODModels[LODIndex].Sections[SectionIndex].MaterialIndex;
 }
 
@@ -46,12 +51,28 @@ FString FSkeletalMeshComponentAdapter::GetBaseName() const
 
 void FSkeletalMeshComponentAdapter::SetMaterial(int32 MaterialIndex, UMaterialInterface* Material)
 {
-	SkeletalMesh->Materials[MaterialIndex] = Material;
+	//Use the material name has the slot name and imported slot name
+	//TODO: find a way to pass the original Material names MaterialSlotName and ImportedMaterialSlotName
+	SkeletalMesh->Materials[MaterialIndex] = FSkeletalMaterial(Material, true, false, Material->GetFName(), Material->GetFName());
 }
 
 void FSkeletalMeshComponentAdapter::RemapMaterialIndex(int32 LODIndex, int32 SectionIndex, int32 NewMaterialIndex)
 {
-	SkeletalMesh->GetImportedModel()->LODModels[LODIndex].Sections[SectionIndex].MaterialIndex = NewMaterialIndex;
+	FSkeletalMeshLODInfo* LODInfoPtr = SkeletalMesh->GetLODInfo(LODIndex);
+	check(LODInfoPtr);
+	if (SkeletalMesh->GetImportedModel()->LODModels[LODIndex].Sections[SectionIndex].MaterialIndex == NewMaterialIndex)
+	{
+		if (LODInfoPtr->LODMaterialMap.IsValidIndex(SectionIndex))
+		{
+			LODInfoPtr->LODMaterialMap[SectionIndex] = INDEX_NONE;
+		}
+	}
+	
+	while (!LODInfoPtr->LODMaterialMap.IsValidIndex(SectionIndex))
+	{
+		LODInfoPtr->LODMaterialMap.Add(INDEX_NONE);
+	}
+	LODInfoPtr->LODMaterialMap[SectionIndex] = NewMaterialIndex;
 }
 
 int32 FSkeletalMeshComponentAdapter::AddMaterial(UMaterialInterface* Material)
diff --git a/Engine/Source/Runtime/Engine/Classes/Engine/SkeletalMesh.h b/Engine/Source/Runtime/Engine/Classes/Engine/SkeletalMesh.h
--- a/Engine/Source/Runtime/Engine/Classes/Engine/SkeletalMesh.h
+++ b/Engine/Source/Runtime/Engine/Classes/Engine/SkeletalMesh.h
@@ -163,7 +163,10 @@ struct FSkeletalMeshLODInfo
 	UPROPERTY(EditAnywhere, Category=SkeletalMeshLODInfo, meta=(DisplayName="LOD Hysteresis"))
 	float LODHysteresis;
 
-	/** Mapping table from this LOD's materials to the USkeletalMesh materials array. */
+	/** Mapping table from this LOD's materials to the USkeletalMesh materials array.
+	 * section index is the key
+	 * remapped material index is the value, can be INDEX_NONE for no remapping
+	 */
 	UPROPERTY()
 	TArray<int32> LODMaterialMap;
 
diff --git a/Engine/Source/Runtime/Engine/Private/SkeletalMeshMerge.cpp b/Engine/Source/Runtime/Engine/Private/SkeletalMeshMerge.cpp
--- a/Engine/Source/Runtime/Engine/Private/SkeletalMeshMerge.cpp
+++ b/Engine/Source/Runtime/Engine/Private/SkeletalMeshMerge.cpp
@@ -302,10 +302,10 @@ void FSkeletalMeshMerge::GenerateNewSectionArray( TArray<FNewSectionInfo>& NewSe
 
 				// get the material for this section
 				int32 MaterialIndex = Section.MaterialIndex;
-				// use the remapping of material indices for all LODs besides the base LOD 
-				if(SrcLODInfo.LODMaterialMap.IsValidIndex(Section.MaterialIndex) && SrcMesh->Materials.IsValidIndex(SrcLODInfo.LODMaterialMap[Section.MaterialIndex]))
+				// use the remapping of material indices if there is a valid value
+				if(SrcLODInfo.LODMaterialMap.IsValidIndex(SectionIdx) && SrcLODInfo.LODMaterialMap[SectionIdx] != INDEX_NONE)
 				{
-					MaterialIndex = FMath::Clamp<int32>( SrcLODInfo.LODMaterialMap[Section.MaterialIndex], 0, SrcMesh->Materials.Num() );
+					MaterialIndex = FMath::Clamp<int32>( SrcLODInfo.LODMaterialMap[SectionIdx], 0, SrcMesh->Materials.Num() );
 				}
 				UMaterialInterface* Material = SrcMesh->Materials[MaterialIndex].MaterialInterface;
 
