HLOD - Meshes marked as useMaxLODAsImposter have an offset in medium HLODs

Fixed mesh pivot point offset not being applied in this specific case

diff --git a/Engine/Source/Developer/MeshMergeUtilities/Private/MeshMergeHelpers.cpp b/Engine/Source/Developer/MeshMergeUtilities/Private/MeshMergeHelpers.cpp
--- a/Engine/Source/Developer/MeshMergeUtilities/Private/MeshMergeHelpers.cpp
+++ b/Engine/Source/Developer/MeshMergeUtilities/Private/MeshMergeHelpers.cpp
@@ -1356,6 +1356,15 @@ void FMeshMergeHelpers::MergeImpostersToRawMesh(TArray<const UStaticMeshComponen
 				VertexInstanceUVs.Set(VertexInstanceID, UVTwoIndex, UVTwo);
 			}
 		}
+		else if (!InPivot.IsZero())
+		{
+			// Apply pivot offset if non null
+			TVertexAttributesRef<FVector> ImposterMeshVertexPositions = ImposterMesh.VertexAttributes().GetAttributesRef<FVector>(MeshAttribute::Vertex::Position);
+			for (FVertexID VertexID : ImposterMesh.Vertices().GetElementIDs())
+			{
+				ImposterMeshVertexPositions[VertexID] -= InPivot;
+			}
+		}
 
 		TPolygonGroupAttributesRef<FName> SourcePolygonGroupImportedMaterialSlotNames = ImposterMeshAttributes.GetPolygonGroupMaterialSlotNames();
 		TPolygonGroupAttributesRef<FName> TargetPolygonGroupImportedMaterialSlotNames = InRawMesh.PolygonGroupAttributes().GetAttributesRef<FName>(MeshAttribute::PolygonGroup::ImportedMaterialSlotName);
