HLOD - Baked Emissive textures do not have sRGB checked
HLOD - Baked gray-scale/mask textures (Rough/Metal/Spec/Opacity) use sRGB=True

diff --git a/Engine/Source/Developer/MeshMergeUtilities/Private/ProxyMaterialUtilities.h b/Engine/Source/Developer/MeshMergeUtilities/Private/ProxyMaterialUtilities.h
--- a/Engine/Source/Developer/MeshMergeUtilities/Private/ProxyMaterialUtilities.h
+++ b/Engine/Source/Developer/MeshMergeUtilities/Private/ProxyMaterialUtilities.h
@@ -163,17 +163,17 @@ namespace ProxyMaterialUtilities
 		// Determine whether or not specific material properties are packed together into one texture (requires at least two to match (number of samples and texture size) in order to be packed
 		if (!bPackMetallic && (FlattenMaterial.GetPropertySize(EFlattenMaterialProperties::Metallic).Num() > 0 || !InMaterialProxySettings.bMetallicMap))
 		{
-			TEXTURE_MACRO_SCALAR(Metallic, TC_Default, bSRGB);
+			TEXTURE_MACRO_SCALAR(Metallic, TC_Default, bRGB);
 		}
 
 		if (!bPackRoughness && (FlattenMaterial.GetPropertySize(EFlattenMaterialProperties::Roughness).Num() > 0 || !InMaterialProxySettings.bRoughnessMap))
 		{
-			TEXTURE_MACRO_SCALAR(Roughness, TC_Default, bSRGB);
+			TEXTURE_MACRO_SCALAR(Roughness, TC_Default, bRGB);
 		}
 
 		if (!bPackSpecular && (FlattenMaterial.GetPropertySize(EFlattenMaterialProperties::Specular).Num() > 0 || !InMaterialProxySettings.bSpecularMap))
 		{
-			TEXTURE_MACRO_SCALAR(Specular, TC_Default, bSRGB);
+			TEXTURE_MACRO_SCALAR(Specular, TC_Default, bRGB);
 		}
 
 		const bool bNonSRGB = false;
@@ -228,7 +228,7 @@ namespace ProxyMaterialUtilities
 			}
 
 			// Create texture using the merged property data
-			UTexture2D* PackedTexture = FMaterialUtilities::CreateTexture(InOuter, AssetBasePath + TEXT("T_") + AssetBaseName + TEXT("_MRS"), PackedSize, MergedTexture, TC_Default, TEXTUREGROUP_HierarchicalLOD, RF_Public | RF_Standalone, bSRGB);
+			UTexture2D* PackedTexture = FMaterialUtilities::CreateTexture(InOuter, AssetBasePath + TEXT("T_") + AssetBaseName + TEXT("_MRS"), PackedSize, MergedTexture, TC_Default, TEXTUREGROUP_HierarchicalLOD, RF_Public | RF_Standalone, bRGB);
 			checkf(PackedTexture, TEXT("Failed to create texture"));
 			OutAssetsToSync.Add(PackedTexture);
 
@@ -257,7 +257,7 @@ namespace ProxyMaterialUtilities
 		// Emissive is a special case due to the scaling variable
 		if (FlattenMaterial.GetPropertySamples(EFlattenMaterialProperties::Emissive).Num() > 0 && !(FlattenMaterial.GetPropertySize(EFlattenMaterialProperties::Emissive).Num() == 1 && FlattenMaterial.GetPropertySamples(EFlattenMaterialProperties::Emissive)[0] == FColor::Black))
 		{
-			TEXTURE_MACRO_VECTOR_LINEAR(Emissive, TC_Default, bRGB);
+			TEXTURE_MACRO_VECTOR_LINEAR(Emissive, TC_Default, bSRGB);
 
 			if (FlattenMaterial.EmissiveScale != 1.0f)
 			{
