Fixed crash when 'merge actors' tries to create a mesh which name is already used by another type of asset

#rnx

diff --git a/Engine/Source/Developer/MeshMergeUtilities/Private/MeshMergeUtilities.cpp b/Engine/Source/Developer/MeshMergeUtilities/Private/MeshMergeUtilities.cpp
--- a/Engine/Source/Developer/MeshMergeUtilities/Private/MeshMergeUtilities.cpp
+++ b/Engine/Source/Developer/MeshMergeUtilities/Private/MeshMergeUtilities.cpp
@@ -2456,6 +2456,32 @@ void FMeshMergeUtilities::MergeComponentsToStaticMesh(const TArray<UPrimitiveCom
 			Package->Modify();
 		}
 
+		// Check that an asset of a different class does not already exist
+		{
+			UObject* ExistingObject = StaticFindObject( nullptr, Package, *AssetName);
+			if(ExistingObject && !ExistingObject->GetClass()->IsChildOf(UStaticMesh::StaticClass()))
+			{
+				// Change name of merged static mesh to avoid name collision
+				UPackage* ParentPackage = CreatePackage(nullptr, *FPaths::GetPath(Package->GetPathName()));
+				ParentPackage->FullyLoad();
+
+				AssetName = MakeUniqueObjectName( ParentPackage, UStaticMesh::StaticClass(), *AssetName).ToString();
+				Package = CreatePackage(NULL, *(ParentPackage->GetPathName() / AssetName ));
+				check(Package);
+				Package->FullyLoad();
+				Package->Modify();
+
+				// Let user know name of merged static mesh has changed
+				UE_LOG(LogMeshMerging, Warning,
+					TEXT("Cannot create %s %s.%s\n")
+					TEXT("An object with the same fully qualified name but a different class already exists.\n")
+					TEXT("\tExisting Object: %s\n")
+					TEXT("The merged mesh will be named %s.%s"),
+					*UStaticMesh::StaticClass()->GetName(), *ExistingObject->GetOutermost()->GetPathName(),	*ExistingObject->GetName(),
+					*ExistingObject->GetFullName(), *Package->GetPathName(), *AssetName);
+			}
+		}
+
 		FStaticMeshComponentRecreateRenderStateContext RecreateRenderStateContext(FindObject<UStaticMesh>(Package, *AssetName));
 
 		UStaticMesh* StaticMesh = NewObject<UStaticMesh>(Package, *AssetName, RF_Public | RF_Standalone);
