Fixed crash when generating HLOD proxy meshes with instanced actors.

#rnx

diff --git a/Engine/Source/Developer/MeshMergeUtilities/Private/MeshMergeHelpers.cpp b/Engine/Source/Developer/MeshMergeUtilities/Private/MeshMergeHelpers.cpp
--- a/Engine/Source/Developer/MeshMergeUtilities/Private/MeshMergeHelpers.cpp
+++ b/Engine/Source/Developer/MeshMergeUtilities/Private/MeshMergeHelpers.cpp
@@ -179,11 +179,21 @@ void FMeshMergeHelpers::ExpandInstances(const UInstancedStaticMeshComponent* InI
 {
 	FMeshDescription CombinedRawMesh;
 
+	bool bFirstMesh = true;
 	for(const FInstancedStaticMeshInstanceData& InstanceData : InInstancedStaticMeshComponent->PerInstanceSMData)
 	{
 		FMeshDescription InstanceRawMesh = InOutRawMesh;
 		FMeshMergeHelpers::TransformRawMeshVertexData(FTransform(InstanceData.Transform), InstanceRawMesh);
-		FMeshMergeHelpers::AppendRawMesh(CombinedRawMesh, InstanceRawMesh);
+
+		if(bFirstMesh)
+		{
+			CombinedRawMesh = InstanceRawMesh;
+			bFirstMesh = false;
+		}
+		else
+		{
+			FMeshMergeHelpers::AppendRawMesh(CombinedRawMesh, InstanceRawMesh);
+		}
 	}
 
 	InOutRawMesh = CombinedRawMesh;
@@ -1185,7 +1195,7 @@ void FMeshMergeHelpers::AppendRawMesh(FMeshDescription& InTarget, const FMeshDes
 	InTarget.ReserveNewVertices(InSource.Vertices().Num());
 	InTarget.ReserveNewVertexInstances(InSource.VertexInstances().Num());
 	InTarget.ReserveNewEdges(InSource.Edges().Num());
-	InTarget.ReserveNewPolygons(InSource.Vertices().Num());
+	InTarget.ReserveNewPolygons(InSource.Polygons().Num());
 
 	//Append PolygonGroup
 	for (const FPolygonGroupID& SourcePolygonGroupID : InSource.PolygonGroups().GetElementIDs())
@@ -1199,12 +1209,14 @@ void FMeshMergeHelpers::AppendRawMesh(FMeshDescription& InTarget, const FMeshDes
 			bool bUnique = true;
 			do 
 			{
+				bUnique = true;
 				for (const FPolygonGroupID PolygonGroupID : InTarget.PolygonGroups().GetElementIDs())
 				{
 					if (TargetPolygonGroupImportedMaterialSlotNames[PolygonGroupID] == CurrentTestName)
 					{
 						CurrentTestName = FName(*(BaseName.ToString() + FString::FromInt(UniqueID++)));
 						bUnique = false;
+						break;
 					}
 				}
 			} while (!bUnique);
