Fix crash when merging staticmesh and using CreateMeshMerge option.
#rnx

diff --git a/Engine/Source/Developer/MeshMergeUtilities/Private/MeshMergeUtilities.cpp b/Engine/Source/Developer/MeshMergeUtilities/Private/MeshMergeUtilities.cpp
--- a/Engine/Source/Developer/MeshMergeUtilities/Private/MeshMergeUtilities.cpp
+++ b/Engine/Source/Developer/MeshMergeUtilities/Private/MeshMergeUtilities.cpp
@@ -2304,7 +2304,7 @@ void FMeshMergeUtilities::MergeComponentsToStaticMesh(const TArray<UPrimitiveCom
 							}
 						}
 
-						if (InSettings.bCreateMergedMaterial)
+						if (InSettings.bCreateMergedMaterial && MeshData.CustomTextureCoordinates.Num() > 0)
 						{
 							VertexInstanceUVs.SetNumIndices(MergedMatUVChannel + 1);
 
@@ -2336,21 +2336,24 @@ void FMeshMergeUtilities::MergeComponentsToStaticMesh(const TArray<UPrimitiveCom
 		OutMaterial.UVChannel = MergedMatUVChannel;
 
 		MergedMaterial = CreateProxyMaterial(InBasePackageName, MergedAssetPackageName, InBaseMaterial, InOuter, InSettings, OutMaterial, OutAssetsToSync);
-
-		if (!InSettings.bCreateMergedMaterial)
+		
+		if (MergedMaterial)
 		{
-			UniqueMaterials.Empty(1);
-			UniqueMaterials.Add(MergedMaterial);
+			if (!InSettings.bCreateMergedMaterial)
+			{
+				UniqueMaterials.Empty(1);
+				UniqueMaterials.Add(MergedMaterial);
 
-			FSectionInfo NewSection;
-			NewSection.Material = MergedMaterial;
-			NewSection.EnabledProperties.Add(GET_MEMBER_NAME_CHECKED(FStaticMeshSection, bCastShadow));
-			DataTracker.AddBakedMaterialSection(NewSection);
-		}
+				FSectionInfo NewSection;
+				NewSection.Material = MergedMaterial;
+				NewSection.EnabledProperties.Add(GET_MEMBER_NAME_CHECKED(FStaticMeshSection, bCastShadow));
+				DataTracker.AddBakedMaterialSection(NewSection);
+			}
 
-		for (IMeshMergeExtension* Extension : MeshMergeExtensions)
-		{
-			Extension->OnCreatedProxyMaterial(StaticMeshComponentsToMerge, MergedMaterial);
+			for (IMeshMergeExtension* Extension : MeshMergeExtensions)
+			{
+				Extension->OnCreatedProxyMaterial(StaticMeshComponentsToMerge, MergedMaterial);
+			}
 		}
 	}
 
@@ -2637,7 +2640,7 @@ void FMeshMergeUtilities::MergeComponentsToStaticMesh(const TArray<UPrimitiveCom
 
 		StaticMesh->PostEditChange();
 
-		if (InSettings.bCreateMergedMaterial)
+		if (InSettings.bCreateMergedMaterial && MergedMaterial)
 		{
 			//Make sure we have unique slot name here
 			FName MaterialSlotName = MergedMaterial->GetFName();
