DRPS - Badly merged files
#rnx

diff --git a/Engine/Source/Developer/MeshMergeUtilities/Private/MeshMergeHelpers.cpp b/Engine/Source/Developer/MeshMergeUtilities/Private/MeshMergeHelpers.cpp
--- a/Engine/Source/Developer/MeshMergeUtilities/Private/MeshMergeHelpers.cpp
+++ b/Engine/Source/Developer/MeshMergeUtilities/Private/MeshMergeHelpers.cpp
@@ -109,17 +109,26 @@ void FMeshMergeHelpers::ExtractSections(const USkeletalMeshComponent* Component,
 
 	TArray<FName> MaterialSlotNames = Component->GetMaterialSlotNames();
 
+	const FSkeletalMeshLODInfo* LODInfoPtr = Component->SkeletalMesh->GetLODInfo(LODIndex);
+	check(LODInfoPtr);
 	const FSkeletalMeshLODModel& Model = Resource->LODModels[LODIndex];
-	for (const FSkelMeshSection& MeshSection : Model.Sections)
+	for (int32 SectionIndex = 0; SectionIndex < Model.Sections.Num(); ++SectionIndex)
 	{
+		const FSkelMeshSection& MeshSection = Model.Sections[SectionIndex];
 		// Retrieve material for this section
-		UMaterialInterface* StoredMaterial = Component->GetMaterial(MeshSection.MaterialIndex);
+		int32 MaterialIndex = MeshSection.MaterialIndex;
+		if (LODInfoPtr->LODMaterialMap.IsValidIndex(SectionIndex) && LODInfoPtr->LODMaterialMap[SectionIndex] != INDEX_NONE)
+		{
+			MaterialIndex = LODInfoPtr->LODMaterialMap[SectionIndex];
+		}
+		
+		UMaterialInterface* StoredMaterial = Component->GetMaterial(MaterialIndex);
 		// Make sure the resource actual exists, otherwise use default material
 		StoredMaterial = (StoredMaterial != nullptr) && StoredMaterial->GetMaterialResource(GMaxRHIFeatureLevel) ? StoredMaterial : DefaultMaterial;
 
 		FSectionInfo SectionInfo;
 		SectionInfo.Material = StoredMaterial;
-		SectionInfo.MaterialSlotName = MaterialSlotNames.IsValidIndex(MeshSection.MaterialIndex) ? MaterialSlotNames[MeshSection.MaterialIndex] : NAME_None;
+		SectionInfo.MaterialSlotName = MaterialSlotNames.IsValidIndex(MaterialIndex) ? MaterialSlotNames[MaterialIndex] : NAME_None;
 
 		if (MeshSection.bCastShadow && Component->CastShadow)
 		{
@@ -309,16 +318,15 @@ void FMeshMergeHelpers::RetrieveMesh(USkeletalMeshComponent* SkeletalMeshCompone
 			const int32 NumWedges = SkelMeshSection.NumTriangles * 3;
 
 			//Create the polygon group ID
-			int32 SectionMaterialIndex = SkelMeshSection.MaterialIndex;
-			int32 MaterialIndex = SectionMaterialIndex;
-			// use the remapping of material indices for all LODs besides the base LOD 
-			if (LODIndex > 0 && SrcLODInfo.LODMaterialMap.IsValidIndex(SkelMeshSection.MaterialIndex))
+			int32 MaterialIndex = SkelMeshSection.MaterialIndex;
+			// use the remapping of material indices if there is a valid value
+			if (SrcLODInfo.LODMaterialMap.IsValidIndex(SectionIndex) && SrcLODInfo.LODMaterialMap[SectionIndex] != INDEX_NONE)
 			{
-				MaterialIndex = FMath::Clamp<int32>(SrcLODInfo.LODMaterialMap[SkelMeshSection.MaterialIndex], 0, SkeletalMeshComponent->SkeletalMesh->Materials.Num());
+				MaterialIndex = FMath::Clamp<int32>(SrcLODInfo.LODMaterialMap[SectionIndex], 0, SkeletalMeshComponent->SkeletalMesh->Materials.Num());
 			}
 
 			FName ImportedMaterialSlotName = SkeletalMeshComponent->SkeletalMesh->Materials[MaterialIndex].ImportedMaterialSlotName;
-			const FPolygonGroupID SectionPolygonGroupID(SectionMaterialIndex);
+			const FPolygonGroupID SectionPolygonGroupID(SectionIndex);
 			if (!RawMesh.IsPolygonGroupValid(SectionPolygonGroupID))
 			{
 				RawMesh.CreatePolygonGroupWithID(SectionPolygonGroupID);
@@ -1348,6 +1356,15 @@ void FMeshMergeHelpers::MergeImpostersToRawMesh(TArray<const UStaticMeshComponen
 				VertexInstanceUVs.Set(VertexInstanceID, UVTwoIndex, UVTwo);
 			}
 		}
+		else if (!InPivot.IsZero())
+		{
+			// Apply pivot offset if non null
+			TVertexAttributesRef<FVector> ImposterMeshVertexPositions = ImposterMesh.VertexAttributes().GetAttributesRef<FVector>(MeshAttribute::Vertex::Position);
+			for (FVertexID VertexID : ImposterMesh.Vertices().GetElementIDs())
+			{
+				ImposterMeshVertexPositions[VertexID] -= InPivot;
+			}
+		}
 
 		TPolygonGroupAttributesRef<FName> SourcePolygonGroupImportedMaterialSlotNames = ImposterMeshAttributes.GetPolygonGroupMaterialSlotNames();
 		TPolygonGroupAttributesRef<FName> TargetPolygonGroupImportedMaterialSlotNames = InRawMesh.PolygonGroupAttributes().GetAttributesRef<FName>(MeshAttribute::PolygonGroup::ImportedMaterialSlotName);
diff --git a/Engine/Source/Developer/MeshMergeUtilities/Private/MeshMergeUtilities.cpp b/Engine/Source/Developer/MeshMergeUtilities/Private/MeshMergeUtilities.cpp
--- a/Engine/Source/Developer/MeshMergeUtilities/Private/MeshMergeUtilities.cpp
+++ b/Engine/Source/Developer/MeshMergeUtilities/Private/MeshMergeUtilities.cpp
@@ -417,6 +417,7 @@ void FMeshMergeUtilities::BakeMaterialsForComponent(TArray<TWeakObjectPtr<UObjec
 			int32 SetIndex = Adapter->GetMaterialIndex(LODIndex, Pair.Key);
 			if (!NonReplaceMaterialIndices.Contains(SetIndex))
 			{
+				//TODO (Bug), need to pass the material data MaterialSlotName and ImportedMaterialSlotName. We loose all this data when baking material on skeletalmesh
 				Adapter->SetMaterial(SetIndex, NewMaterials[Pair.Value]);
 			}
 			else
diff --git a/Engine/Source/Developer/MeshMergeUtilities/Private/SkeletalMeshAdapter.cpp b/Engine/Source/Developer/MeshMergeUtilities/Private/SkeletalMeshAdapter.cpp
--- a/Engine/Source/Developer/MeshMergeUtilities/Private/SkeletalMeshAdapter.cpp
+++ b/Engine/Source/Developer/MeshMergeUtilities/Private/SkeletalMeshAdapter.cpp
@@ -31,6 +31,11 @@ void FSkeletalMeshComponentAdapter::RetrieveMeshSections(int32 LODIndex, TArray<
 
 int32 FSkeletalMeshComponentAdapter::GetMaterialIndex(int32 LODIndex, int32 SectionIndex) const
 {
+	const FSkeletalMeshLODInfo* LODInfoPtr = SkeletalMesh->GetLODInfo(LODIndex);
+	if (LODInfoPtr && LODInfoPtr->LODMaterialMap.IsValidIndex(SectionIndex) && LODInfoPtr->LODMaterialMap[SectionIndex] != INDEX_NONE)
+	{
+		return LODInfoPtr->LODMaterialMap[SectionIndex];
+	}
 	return SkeletalMesh->GetImportedModel()->LODModels[LODIndex].Sections[SectionIndex].MaterialIndex;
 }
 
@@ -46,12 +51,28 @@ FString FSkeletalMeshComponentAdapter::GetBaseName() const
 
 void FSkeletalMeshComponentAdapter::SetMaterial(int32 MaterialIndex, UMaterialInterface* Material)
 {
-	SkeletalMesh->Materials[MaterialIndex] = Material;
+	//Use the material name has the slot name and imported slot name
+	//TODO: find a way to pass the original Material names MaterialSlotName and ImportedMaterialSlotName
+	SkeletalMesh->Materials[MaterialIndex] = FSkeletalMaterial(Material, true, false, Material->GetFName(), Material->GetFName());
 }
 
 void FSkeletalMeshComponentAdapter::RemapMaterialIndex(int32 LODIndex, int32 SectionIndex, int32 NewMaterialIndex)
 {
-	SkeletalMesh->GetImportedModel()->LODModels[LODIndex].Sections[SectionIndex].MaterialIndex = NewMaterialIndex;
+	FSkeletalMeshLODInfo* LODInfoPtr = SkeletalMesh->GetLODInfo(LODIndex);
+	check(LODInfoPtr);
+	if (SkeletalMesh->GetImportedModel()->LODModels[LODIndex].Sections[SectionIndex].MaterialIndex == NewMaterialIndex)
+	{
+		if (LODInfoPtr->LODMaterialMap.IsValidIndex(SectionIndex))
+		{
+			LODInfoPtr->LODMaterialMap[SectionIndex] = INDEX_NONE;
+		}
+	}
+	
+	while (!LODInfoPtr->LODMaterialMap.IsValidIndex(SectionIndex))
+	{
+		LODInfoPtr->LODMaterialMap.Add(INDEX_NONE);
+	}
+	LODInfoPtr->LODMaterialMap[SectionIndex] = NewMaterialIndex;
 }
 
 int32 FSkeletalMeshComponentAdapter::AddMaterial(UMaterialInterface* Material)
diff --git a/Engine/Source/Developer/MeshUtilities/Private/MeshUtilities.cpp b/Engine/Source/Developer/MeshUtilities/Private/MeshUtilities.cpp
--- a/Engine/Source/Developer/MeshUtilities/Private/MeshUtilities.cpp
+++ b/Engine/Source/Developer/MeshUtilities/Private/MeshUtilities.cpp
@@ -303,10 +303,10 @@ static void SkinnedMeshToRawMeshes(USkinnedMeshComponent* InSkinnedMeshComponent
 				}
 
 				int32 MaterialIndex = SkelMeshSection.MaterialIndex;
-				// use the remapping of material indices for all LODs besides the base LOD 
-				if (LODIndexRead > 0 && SrcLODInfo.LODMaterialMap.IsValidIndex(SkelMeshSection.MaterialIndex))
+				// use the remapping of material indices if there is a valid value
+				if (SrcLODInfo.LODMaterialMap.IsValidIndex(SectionIndex) && SrcLODInfo.LODMaterialMap[SectionIndex] != INDEX_NONE)
 				{
-					MaterialIndex = FMath::Clamp<int32>(SrcLODInfo.LODMaterialMap[SkelMeshSection.MaterialIndex], 0, InSkinnedMeshComponent->SkeletalMesh->Materials.Num());
+					MaterialIndex = FMath::Clamp<int32>(SrcLODInfo.LODMaterialMap[SectionIndex], 0, InSkinnedMeshComponent->SkeletalMesh->Materials.Num());
 				}
 
 				// copy face info
diff --git a/Engine/Source/Developer/ToolMenus/Private/ToolMenus.cpp b/Engine/Source/Developer/ToolMenus/Private/ToolMenus.cpp
--- a/Engine/Source/Developer/ToolMenus/Private/ToolMenus.cpp
+++ b/Engine/Source/Developer/ToolMenus/Private/ToolMenus.cpp
@@ -597,6 +597,23 @@ void UToolMenus::ApplyCustomization(UToolMenu* GeneratedMenu)
 		}
 	}
 
+	// Hide items based on whitelist
+	if (CustomizedMenu.WhitelistEntries.Num() > 0)
+	{
+		for (int32 SectionIndex = 0; SectionIndex < NewSections.Num(); ++SectionIndex)
+		{
+			FToolMenuSection& Section = NewSections[SectionIndex];
+			for (int32 i = 0; i < Section.Blocks.Num(); ++i)
+			{
+				if (!CustomizedMenu.IsEntryWhitelisted(Section.Blocks[i].Name))
+				{
+					Section.Blocks.RemoveAt(i);
+					--i;
+				}
+			}
+		}
+	}
+
 	// Hide sections and entries
 	if (!GeneratedMenu->IsEditing())
 	{
@@ -1421,6 +1438,11 @@ UToolMenu* UToolMenus::GenerateMenuOrSubMenuForEdit(const UToolMenu* InMenu)
 	return nullptr;
 }
 
+void UToolMenus::AddMenuSubstitutionDuringGenerate(const FName OriginalMenu, const FName NewMenu)
+{
+	MenuSubstitutionsDuringGenerate.Add(OriginalMenu, NewMenu);
+}
+
 UToolMenu* UToolMenus::GenerateMenu(const FName Name, const FToolMenuContext& InMenuContext)
 {
 	// Allow substituting one menu for another
diff --git a/Engine/Source/Developer/ToolMenus/Public/ToolMenus.h b/Engine/Source/Developer/ToolMenus/Public/ToolMenus.h
--- a/Engine/Source/Developer/ToolMenus/Public/ToolMenus.h
+++ b/Engine/Source/Developer/ToolMenus/Public/ToolMenus.h
@@ -260,6 +260,9 @@ public:
 	/** Enables adding command to open edit menu dialog to each menu */
 	void SetEditMenusMode(bool bEnable);
 
+	/* Substitute one menu for another during generate but not during find or extend */
+	void AddMenuSubstitutionDuringGenerate(const FName OriginalMenu, const FName NewMenu);
+
 	/** Displaying extension points is for debugging menus */
 	DECLARE_DELEGATE_RetVal(bool, FShouldDisplayExtensionPoints);
 	FShouldDisplayExtensionPoints ShouldDisplayExtensionPoints;
@@ -322,7 +325,7 @@ private:
 	void PopulateMenuBarBuilder(FMenuBarBuilder& MenuBarBuilder, UToolMenu* MenuData);
 	void PopulateToolBarBuilder(FToolBarBuilder& ToolBarBuilder, UToolMenu* MenuData);
 
-	TSharedRef<SWidget> GenerateToolbarComboButtonMenu(const FName SubMenuFullName, FToolMenuContext InContext);
+	TSharedRef<SWidget> GenerateToolbarComboButtonMenu(TWeakObjectPtr<UToolMenu> InParent, const FName InBlockName);
 
 	FOnGetContent ConvertWidgetChoice(const FNewToolMenuWidgetChoice& Choice, const FToolMenuContext& Context) const;
 
