Integrate 11534852 from 4.25: Add missing flattened material property enum that was missed when anisotropy was added.

diff --git a/Engine/Source/Developer/MeshMergeUtilities/Private/MeshMergeUtilities.cpp b/Engine/Source/Developer/MeshMergeUtilities/Private/MeshMergeUtilities.cpp
--- a/Engine/Source/Developer/MeshMergeUtilities/Private/MeshMergeUtilities.cpp
+++ b/Engine/Source/Developer/MeshMergeUtilities/Private/MeshMergeUtilities.cpp
@@ -600,7 +600,9 @@ EFlattenMaterialProperties FMeshMergeUtilities::NewToOldProperty(int32 NewProper
 		EFlattenMaterialProperties::Metallic,
 		EFlattenMaterialProperties::Specular,
 		EFlattenMaterialProperties::Roughness,
+		EFlattenMaterialProperties::Anisotropy,
 		EFlattenMaterialProperties::Normal,
+		EFlattenMaterialProperties::Tangent,
 		EFlattenMaterialProperties::NumFlattenMaterialProperties,
 		EFlattenMaterialProperties::NumFlattenMaterialProperties,
 		EFlattenMaterialProperties::NumFlattenMaterialProperties,
