Avoid memcopies in MergeFlattenedMaterials.

diff --git a/Engine/Source/Developer/MeshMergeUtilities/Private/MeshMergeUtilities.cpp b/Engine/Source/Developer/MeshMergeUtilities/Private/MeshMergeUtilities.cpp
--- a/Engine/Source/Developer/MeshMergeUtilities/Private/MeshMergeUtilities.cpp
+++ b/Engine/Source/Developer/MeshMergeUtilities/Private/MeshMergeUtilities.cpp
@@ -868,9 +868,8 @@ void FMeshMergeUtilities::MergeFlattenedMaterials(TArray<struct FFlattenMaterial
 	OutUVTransforms.Reserve(InMaterialList.Num());
 
 	// Fill output UV transforms with invalid values
-	for (auto Material : InMaterialList)
+	for (auto& Material : InMaterialList)
 	{
-
 		// Invalid UV transform
 		FUVOffsetScalePair UVTransform;
 		UVTransform.Key = FVector2D::ZeroVector;
