Fix build error in MeshMergeUtilities.cpp

diff --git a/Engine/Source/Developer/MeshMergeUtilities/Private/MeshMergeUtilities.cpp b/Engine/Source/Developer/MeshMergeUtilities/Private/MeshMergeUtilities.cpp
--- a/Engine/Source/Developer/MeshMergeUtilities/Private/MeshMergeUtilities.cpp
+++ b/Engine/Source/Developer/MeshMergeUtilities/Private/MeshMergeUtilities.cpp
@@ -2928,7 +2928,7 @@ void FMeshMergeUtilities::CreateMergedRawMeshes(FMeshMergeDataTracker& InDataTra
 					});
 					AppendSettings.bMergeVertexColor = InSettings.bBakeVertexDataToMesh;
 					AppendSettings.MergedAssetPivot = InMergedAssetPivot;
-					for (int32 ChannelIdx = 0; ChannelIdx < FMeshDescriptionOperations::FAppendSettings::MAX_MESH_TEXTURE_COORDS; ++ChannelIdx)
+					for (int32 ChannelIdx = 0; ChannelIdx < FStaticMeshOperations::FAppendSettings::MAX_MESH_TEXTURE_COORDS; ++ChannelIdx)
 					{
 						AppendSettings.bMergeUVChannels[ChannelIdx] = InDataTracker.DoesUVChannelContainData(ChannelIdx, LODIndex) && InSettings.OutputUVs[ChannelIdx] == EUVOutput::OutputChannel;
 					}
@@ -3037,7 +3037,7 @@ void FMeshMergeUtilities::CreateMergedRawMeshes(FMeshMergeDataTracker& InDataTra
 				});
 				AppendSettings.bMergeVertexColor = InSettings.bBakeVertexDataToMesh;
 				AppendSettings.MergedAssetPivot = InMergedAssetPivot;
-				for (int32 ChannelIdx = 0; ChannelIdx < FMeshDescriptionOperations::FAppendSettings::MAX_MESH_TEXTURE_COORDS; ++ChannelIdx)
+				for (int32 ChannelIdx = 0; ChannelIdx < FStaticMeshOperations::FAppendSettings::MAX_MESH_TEXTURE_COORDS; ++ChannelIdx)
 				{
 					AppendSettings.bMergeUVChannels[ChannelIdx] = InDataTracker.DoesUVChannelContainData(ChannelIdx, LODIndex) && InSettings.OutputUVs[ChannelIdx] == EUVOutput::OutputChannel;
 				}
