CIS fix

diff --git a/Engine/Source/Developer/MeshMergeUtilities/Private/MeshMergeUtilities.cpp b/Engine/Source/Developer/MeshMergeUtilities/Private/MeshMergeUtilities.cpp
--- a/Engine/Source/Developer/MeshMergeUtilities/Private/MeshMergeUtilities.cpp
+++ b/Engine/Source/Developer/MeshMergeUtilities/Private/MeshMergeUtilities.cpp
@@ -2928,7 +2928,7 @@ void FMeshMergeUtilities::CreateMergedRawMeshes(FMeshMergeDataTracker& InDataTra
 					});
 					AppendSettings.bMergeVertexColor = InSettings.bBakeVertexDataToMesh;
 					AppendSettings.MergedAssetPivot = InMergedAssetPivot;
-					for (int32 ChannelIdx = 0; ChannelIdx < FStaticMeshOperations::FAppendSettings::MAX_MESH_TEXTURE_COORDS; ++ChannelIdx)
+					for (int32 ChannelIdx = 0; ChannelIdx < FStaticMeshOperations::FAppendSettings::MAX_NUM_UV_CHANNELS; ++ChannelIdx)
 					{
 						AppendSettings.bMergeUVChannels[ChannelIdx] = InDataTracker.DoesUVChannelContainData(ChannelIdx, LODIndex) && InSettings.OutputUVs[ChannelIdx] == EUVOutput::OutputChannel;
 					}
@@ -3037,7 +3037,7 @@ void FMeshMergeUtilities::CreateMergedRawMeshes(FMeshMergeDataTracker& InDataTra
 				});
 				AppendSettings.bMergeVertexColor = InSettings.bBakeVertexDataToMesh;
 				AppendSettings.MergedAssetPivot = InMergedAssetPivot;
-				for (int32 ChannelIdx = 0; ChannelIdx < FStaticMeshOperations::FAppendSettings::MAX_MESH_TEXTURE_COORDS; ++ChannelIdx)
+				for (int32 ChannelIdx = 0; ChannelIdx < FStaticMeshOperations::FAppendSettings::MAX_NUM_UV_CHANNELS; ++ChannelIdx)
 				{
 					AppendSettings.bMergeUVChannels[ChannelIdx] = InDataTracker.DoesUVChannelContainData(ChannelIdx, LODIndex) && InSettings.OutputUVs[ChannelIdx] == EUVOutput::OutputChannel;
 				}
diff --git a/Engine/Source/Runtime/StaticMeshDescription/Public/StaticMeshOperations.h b/Engine/Source/Runtime/StaticMeshDescription/Public/StaticMeshOperations.h
--- a/Engine/Source/Runtime/StaticMeshDescription/Public/StaticMeshOperations.h
+++ b/Engine/Source/Runtime/StaticMeshDescription/Public/StaticMeshOperations.h
@@ -43,12 +43,12 @@ public:
 
 		enum
 		{
-			MAX_MESH_TEXTURE_COORDS = 8,
+			MAX_NUM_UV_CHANNELS = 8,
 		};
 
 		FAppendPolygonGroupsDelegate PolygonGroupsDelegate;
 		bool bMergeVertexColor;
-		bool bMergeUVChannels[MAX_MESH_TEXTURE_COORDS];
+		bool bMergeUVChannels[MAX_NUM_UV_CHANNELS];
 		FVector MergedAssetPivot;
 		TOptional<FTransform> MeshTransform; // Apply a transformation on source mesh (see MeshTransform)
 	};
