Fix crash during proxy lods generation if lightmap uvindex is -1

#rnx

diff --git a/Engine/Source/Developer/MeshMergeUtilities/Private/MeshMergeUtilities.cpp b/Engine/Source/Developer/MeshMergeUtilities/Private/MeshMergeUtilities.cpp
--- a/Engine/Source/Developer/MeshMergeUtilities/Private/MeshMergeUtilities.cpp
+++ b/Engine/Source/Developer/MeshMergeUtilities/Private/MeshMergeUtilities.cpp
@@ -1383,23 +1383,26 @@ void FMeshMergeUtilities::CreateProxyMesh(const TArray<UStaticMeshComponent*>& I
 					TArray<uint32> MeshIndices;
 					SectionToMesh.MultiFind(SectionIndex, MeshIndices);
 
-					for (uint32 MeshIndex : MeshIndices)
+					for (const uint32 MeshIndex : MeshIndices)
 					{
-						FMeshData MeshSettings;
 						// Retrieve mesh description
+						const UStaticMeshComponent* StaticMeshComponent = StaticMeshComponents[MeshIndex];
 						FMeshDescription* MeshDescription = MeshDescriptionData[MeshIndex];
+
+						FMeshData MeshSettings;
 						MeshSettings.RawMeshDescription = MeshDescription;
 
-						TVertexInstanceAttributesRef<FVector2D> VertexInstanceUVs = MeshSettings.RawMeshDescription->VertexInstanceAttributes().GetAttributesRef<FVector2D>(MeshAttribute::VertexInstance::TextureCoordinate);
+						TVertexInstanceAttributesRef<FVector2D> VertexInstanceUVs = MeshDescription->VertexInstanceAttributes().GetAttributesRef<FVector2D>(MeshAttribute::VertexInstance::TextureCoordinate);
 
-						// If we already have lightmap uvs generated or the lightmap coordinate index != 0 and available we can reuse those instead of having to generate new ones
-						if (InMeshProxySettings.bReuseMeshLightmapUVs
-							&& (StaticMeshComponents[MeshIndex]->GetStaticMesh()->GetSourceModel(0).BuildSettings.bGenerateLightmapUVs
-								|| (StaticMeshComponents[MeshIndex]->GetStaticMesh()->LightMapCoordinateIndex != 0 && VertexInstanceUVs.GetNumElements() > 0 && VertexInstanceUVs.GetNumIndices() > StaticMeshComponents[MeshIndex]->GetStaticMesh()->LightMapCoordinateIndex)))
+						// If we already have lightmap uvs generated and they are valid, we can reuse those instead of having to generate new ones
+						const int32 LightMapCoordinateIndex = StaticMeshComponent->GetStaticMesh()->LightMapCoordinateIndex;
+						if (InMeshProxySettings.bReuseMeshLightmapUVs &&
+							LightMapCoordinateIndex > 0 &&
+							VertexInstanceUVs.GetNumElements() > 0 &&
+							VertexInstanceUVs.GetNumIndices() > LightMapCoordinateIndex)
 						{
 							MeshSettings.CustomTextureCoordinates.Reset(VertexInstanceUVs.GetNumElements());
-							int32 LightMapCoordinateIndex = StaticMeshComponents[MeshIndex]->GetStaticMesh()->LightMapCoordinateIndex;
-							for (const FVertexInstanceID VertexInstanceID : MeshSettings.RawMeshDescription->VertexInstances().GetElementIDs())
+							for (const FVertexInstanceID VertexInstanceID : MeshDescription->VertexInstances().GetElementIDs())
 							{
 								MeshSettings.CustomTextureCoordinates.Add(VertexInstanceUVs.Get(VertexInstanceID, LightMapCoordinateIndex));
 							}
@@ -1436,7 +1439,6 @@ void FMeshMergeUtilities::CreateProxyMesh(const TArray<UStaticMeshComponent*>& I
 						}
 
 						// Retrieve lightmap for usage of lightmap data
-						const UStaticMeshComponent* StaticMeshComponent = StaticMeshComponents[MeshIndex];
 						if (StaticMeshComponent->LODData.IsValidIndex(0))
 						{
 							const FStaticMeshComponentLODInfo& ComponentLODInfo = StaticMeshComponent->LODData[0];
