Fixed two issues in the proxy merge tools.
1- The FExportMaterialProxy didn't give the proper static parameters for the materials instance. It always used those from the parent/master material.
2- The mesh description extracted from the static mesh component was sometimes invalid. This occurred when the mesh had an empty section as it first section. This resulted in a mesh with all its polygons being part of the first section when rendered for the proxy mesh.

Also renamed a variable in the create proxy mesh function since its previous name didn't make sense with what it is.

diff --git a/Engine/Source/Developer/MaterialBaking/Private/ExportMaterialProxy.h b/Engine/Source/Developer/MaterialBaking/Private/ExportMaterialProxy.h
--- a/Engine/Source/Developer/MaterialBaking/Private/ExportMaterialProxy.h
+++ b/Engine/Source/Developer/MaterialBaking/Private/ExportMaterialProxy.h
@@ -301,6 +301,14 @@ public:
 		return ReferencedTextures;
 	}
 
+	virtual void GetStaticParameterSet(EShaderPlatform Platform, FStaticParameterSet& OutSet) const override
+	{
+		if (const FMaterialResource* Resource = MaterialInterface->GetMaterialResource(GMaxRHIFeatureLevel))
+		{
+			Resource->GetStaticParameterSet(Platform, OutSet);
+		}
+	}
+
 	////////////////
 	// FMaterialRenderProxy interface.
 	virtual const FMaterial& GetMaterialWithFallback(ERHIFeatureLevel::Type FeatureLevel, const FMaterialRenderProxy*& OutFallbackMaterialRenderProxy) const override
diff --git a/Engine/Source/Developer/MeshMergeUtilities/Private/MeshMergeHelpers.cpp b/Engine/Source/Developer/MeshMergeUtilities/Private/MeshMergeHelpers.cpp
--- a/Engine/Source/Developer/MeshMergeUtilities/Private/MeshMergeHelpers.cpp
+++ b/Engine/Source/Developer/MeshMergeUtilities/Private/MeshMergeHelpers.cpp
@@ -487,9 +487,9 @@ void FMeshMergeHelpers::ExportStaticMeshLOD(const FStaticMeshLODResources& Stati
 		for (int32 SectionIndex = 0; SectionIndex < StaticMeshLOD.Sections.Num(); ++SectionIndex)
 		{
 			const FStaticMeshSection& Section = StaticMeshLOD.Sections[SectionIndex];
-			uint32 FirstTriangle = Section.FirstIndex / 3;
-			uint32 LastTriangle = FirstTriangle + Section.NumTriangles - 1;
-			if ((uint32)TriangleIndex >= FirstTriangle && (uint32)TriangleIndex <= LastTriangle)
+			uint32 BeginTriangle = Section.FirstIndex / 3;
+			uint32 EndTriangle = BeginTriangle + Section.NumTriangles;
+			if ((uint32)TriangleIndex >= BeginTriangle && (uint32)TriangleIndex < EndTriangle)
 			{
 				CurrentPolygonGroupID = FPolygonGroupID(SectionIndex);
 				break;
diff --git a/Engine/Source/Developer/MeshMergeUtilities/Private/MeshMergeUtilities.cpp b/Engine/Source/Developer/MeshMergeUtilities/Private/MeshMergeUtilities.cpp
--- a/Engine/Source/Developer/MeshMergeUtilities/Private/MeshMergeUtilities.cpp
+++ b/Engine/Source/Developer/MeshMergeUtilities/Private/MeshMergeUtilities.cpp
@@ -1325,12 +1325,12 @@ void FMeshMergeUtilities::CreateProxyMesh(const TArray<UStaticMeshComponent*>& I
 
 	TArray<UMaterialInterface*> UniqueMaterials;
 	//Unique material index to unique section index
-	TMultiMap<uint32, uint32> SectionToMaterialMap;
+	TMultiMap<uint32, uint32> MaterialToSectionMap;
 	for (int32 SectionIndex = 0; SectionIndex < UniqueSections.Num(); ++SectionIndex)
 	{
 		FSectionInfo& Section = UniqueSections[SectionIndex];
 		const int32 UniqueIndex = UniqueMaterials.AddUnique(Section.Material);
-		SectionToMaterialMap.Add(UniqueIndex, SectionIndex);
+		MaterialToSectionMap.Add(UniqueIndex, SectionIndex);
 	}
 
 	TArray<FMeshData> GlobalMeshSettings;
@@ -1358,7 +1358,7 @@ void FMeshMergeUtilities::CreateProxyMesh(const TArray<UStaticMeshComponent*>& I
 
 			//Unique section indices
 			TArray<uint32> SectionIndices;
-			SectionToMaterialMap.MultiFind(MaterialIndex, SectionIndices);
+			MaterialToSectionMap.MultiFind(MaterialIndex, SectionIndices);
 
 			// Check whether or not this material requires mesh data
 			int32 NumTexCoords = 0;
