First pass of MeshDescription API and format refactor.

diff --git a/Engine/Plugins/Runtime/EditableMesh/Source/EditableMesh/EditableStaticMeshAdapter.cpp b/Engine/Plugins/Runtime/EditableMesh/Source/EditableMesh/EditableStaticMeshAdapter.cpp
--- a/Engine/Plugins/Runtime/EditableMesh/Source/EditableMesh/EditableStaticMeshAdapter.cpp
+++ b/Engine/Plugins/Runtime/EditableMesh/Source/EditableMesh/EditableStaticMeshAdapter.cpp
@@ -85,12 +85,10 @@ void UEditableStaticMeshAdapter::InitEditableStaticMesh( UEditableMesh* Editable
 
 				TPolygonGroupAttributesRef<FName> PolygonGroupImportedMaterialSlotNames = MeshDescription->PolygonGroupAttributes().GetAttributesRef<FName>( MeshAttribute::PolygonGroup::ImportedMaterialSlotName );
 				TPolygonGroupAttributesRef<FName> PolygonGroupMaterialAssetNames = MeshDescription->PolygonGroupAttributes().GetAttributesRef<FName>( MeshAttribute::PolygonGroup::MaterialAssetName );
-				TPolygonGroupAttributesRef<bool> PolygonGroupCollision = MeshDescription->PolygonGroupAttributes().GetAttributesRef<bool>( MeshAttribute::PolygonGroup::EnableCollision );
-				TPolygonGroupAttributesRef<bool> PolygonGroupCastShadow = MeshDescription->PolygonGroupAttributes().GetAttributesRef<bool>( MeshAttribute::PolygonGroup::CastShadow );
 
 				// Store off the number of texture coordinates in this mesh
 				const int32 NumUVs = StaticMeshLOD.GetNumTexCoords();
-				VertexInstanceUVs.SetNumIndices(NumUVs);
+				VertexInstanceUVs.SetNumChannels(NumUVs);
 				EditableMesh->TextureCoordinateCount = NumUVs;
 
 				// Vertices
@@ -261,8 +259,6 @@ void UEditableStaticMeshAdapter::InitEditableStaticMesh( UEditableMesh* Editable
 					const FPolygonGroupID NewPolygonGroupID = MeshDescription->CreatePolygonGroup();
 					PolygonGroupImportedMaterialSlotNames[NewPolygonGroupID] = MaterialSlotName;
 					PolygonGroupMaterialAssetNames[NewPolygonGroupID] = MaterialAssetName;
-					PolygonGroupCollision[ NewPolygonGroupID ] = RenderingSection.bEnableCollision;
-					PolygonGroupCastShadow[ NewPolygonGroupID ] = RenderingSection.bCastShadow;
 
 					// Create a rendering polygon group for holding the triangulated data and references to the static mesh rendering section.
 					// This is indexed by the same FPolygonGroupID as the PolygonGroups.
@@ -308,7 +304,7 @@ void UEditableStaticMeshAdapter::InitEditableStaticMesh( UEditableMesh* Editable
 							const FTriangleID NewTriangleID = FTriangleID( SectionTriangleIndex );
 
 							NewRenderingPolygonGroup.Triangles.Insert( NewTriangleID );
-							FMeshTriangle& NewTriangle = NewRenderingPolygonGroup.Triangles[ NewTriangleID ];
+							FRenderingPolygonGroup::FMeshTriangle& NewTriangle = NewRenderingPolygonGroup.Triangles[ NewTriangleID ];
 
 							for( uint32 TriangleVertexIndex = 0; TriangleVertexIndex < 3; ++TriangleVertexIndex )
 							{
@@ -426,10 +422,10 @@ void UEditableStaticMeshAdapter::InitializeFromEditableMesh( const UEditableMesh
 		FRenderingPolygon& RenderingPolygon = RenderingPolygons[ PolygonID ];
 		RenderingPolygon.PolygonGroupID = PolygonGroupID;
 
-		const TArray<FTriangleID>& TriangleIDs = MeshDescription->GetPolygonTriangleIDs( PolygonID );
+		TArrayView<const FTriangleID> TriangleIDs = MeshDescription->GetPolygonTriangleIDs( PolygonID );
 		for( const FTriangleID TriangleID : TriangleIDs )
 		{
-			FMeshTriangle Triangle;
+			FRenderingPolygonGroup::FMeshTriangle Triangle;
 			for( int32 TriIndex = 0; TriIndex < 3; ++TriIndex )
 			{
 				Triangle.SetVertexInstanceID( TriIndex, MeshDescription->GetTriangleVertexInstance( TriangleID, TriIndex ) );
@@ -465,8 +461,6 @@ void UEditableStaticMeshAdapter::OnRebuildRenderMesh( const UEditableMesh* Edita
 	check( MeshDescription );
 
 	TPolygonGroupAttributesConstRef<FName> PolygonGroupImportedMaterialSlotNames = MeshDescription->PolygonGroupAttributes().GetAttributesRef<FName>( MeshAttribute::PolygonGroup::ImportedMaterialSlotName );
-	TPolygonGroupAttributesConstRef<bool> PolygonGroupCollision = MeshDescription->PolygonGroupAttributes().GetAttributesRef<bool>( MeshAttribute::PolygonGroup::EnableCollision );
-	TPolygonGroupAttributesConstRef<bool> PolygonGroupCastShadow = MeshDescription->PolygonGroupAttributes().GetAttributesRef<bool>( MeshAttribute::PolygonGroup::CastShadow );
 
 	const int32 NumPolygonGroups = MeshDescription->PolygonGroups().Num();
 	StaticMeshLOD.Sections.Empty( NumPolygonGroups );
@@ -509,8 +503,8 @@ void UEditableStaticMeshAdapter::OnRebuildRenderMesh( const UEditableMesh* Edita
 			check( MaterialIndex != INDEX_NONE );
 
 			StaticMeshSection.MaterialIndex = MaterialIndex;
-			StaticMeshSection.bEnableCollision = PolygonGroupCollision[ PolygonGroupID ];
-			StaticMeshSection.bCastShadow = PolygonGroupCastShadow[ PolygonGroupID ];
+			StaticMeshSection.bEnableCollision = true;
+			StaticMeshSection.bCastShadow = true;
 
 			// Fill vertices
 			int32 NextVertexIndex = FirstSectionVertexIndex;
@@ -595,7 +589,7 @@ void UEditableStaticMeshAdapter::OnRebuildRenderMesh( const UEditableMesh* Edita
 			StaticMeshVertex.TangentY = FVector::CrossProduct( VertexInstanceNormals[ VertexInstanceID ], VertexInstanceTangents[ VertexInstanceID ] ).GetSafeNormal() * VertexInstanceBinormalSigns[ VertexInstanceID ];
 			StaticMeshVertex.TangentZ = VertexInstanceNormals[ VertexInstanceID ];
 			StaticMeshVertex.Color = VertexInstanceColor.ToFColor( true );
-			for( int32 UVIndex = 0; UVIndex < VertexInstanceUVs.GetNumIndices(); ++UVIndex )
+			for( int32 UVIndex = 0; UVIndex < VertexInstanceUVs.GetNumChannels(); ++UVIndex )
 			{
 				StaticMeshVertex.UVs[ UVIndex ] = VertexInstanceUVs.Get( VertexInstanceID, UVIndex );
 			}
@@ -628,8 +622,8 @@ void UEditableStaticMeshAdapter::OnRebuildRenderMesh( const UEditableMesh* Edita
 			const int32 MaterialIndex = StaticMesh->GetMaterialIndexFromImportedMaterialSlotName( PolygonGroupImportedMaterialSlotNames[ PolygonGroupID ] );
 			//check( MaterialIndex != INDEX_NONE );
 			StaticMeshSection.MaterialIndex = MaterialIndex;
-			StaticMeshSection.bEnableCollision = PolygonGroupCollision[ PolygonGroupID ];
-			StaticMeshSection.bCastShadow = PolygonGroupCastShadow[ PolygonGroupID ];
+			StaticMeshSection.bEnableCollision = true;
+			StaticMeshSection.bCastShadow = true;
 
 
 			if( RenderingPolygonGroup.Triangles.Num() > 0 )
@@ -647,7 +641,7 @@ void UEditableStaticMeshAdapter::OnRebuildRenderMesh( const UEditableMesh* Edita
 					const FTriangleID TriangleID( TriangleIndex );
 					if( RenderingPolygonGroup.Triangles.IsValid( TriangleID ) )
 					{
-						const FMeshTriangle& Triangle = RenderingPolygonGroup.Triangles[ TriangleID ];
+						const FRenderingPolygonGroup::FMeshTriangle& Triangle = RenderingPolygonGroup.Triangles[ TriangleID ];
 						for( int32 TriVert = 0; TriVert < 3; ++TriVert )
 						{
 							const uint32 RenderingVertexIndex = Triangle.GetVertexInstanceID( TriVert ).GetValue();
@@ -826,7 +820,7 @@ void UEditableStaticMeshAdapter::OnReindexElements( const UEditableMesh* Editabl
 
 		for( const FTriangleID TriangleID : RenderingPolygonGroup.Triangles.GetElementIDs() )
 		{
-			FMeshTriangle& Triangle = RenderingPolygonGroup.Triangles[ TriangleID ];
+			FRenderingPolygonGroup::FMeshTriangle& Triangle = RenderingPolygonGroup.Triangles[ TriangleID ];
 			for( int32 TriangleVertexNumber = 0; TriangleVertexNumber < 3; ++TriangleVertexNumber )
 			{
 				const FVertexInstanceID OriginalVertexInstanceID = Triangle.GetVertexInstanceID( TriangleVertexNumber );
@@ -1194,7 +1188,7 @@ void UEditableStaticMeshAdapter::OnCreateVertexInstances( const UEditableMesh* E
 	{
 		const FMeshDescription* MeshDescription = EditableMesh->GetMeshDescription();
 		FStaticMeshLODResources& StaticMeshLOD = GetStaticMeshLOD();
-		const int32 NumUVs = MeshDescription->VertexInstanceAttributes().GetAttributeIndexCount<FVector2D>( MeshAttribute::VertexInstance::TextureCoordinate );
+		const int32 NumUVs = MeshDescription->VertexInstanceAttributes().GetAttributeChannelCount( MeshAttribute::VertexInstance::TextureCoordinate );
 		TVertexAttributesConstRef<FVector> VertexPositions = MeshDescription->VertexAttributes().GetAttributesRef<FVector>( MeshAttribute::Vertex::Position );
 		const bool bHasColors = StaticMeshLOD.VertexBuffers.ColorVertexBuffer.GetNumVertices() > 0;
 
@@ -1302,7 +1296,7 @@ void UEditableStaticMeshAdapter::OnRetriangulatePolygons( const UEditableMesh* E
 		FRenderingPolygon& RenderingPolygon = RenderingPolygons[ PolygonID ];
 		const FPolygonGroupID PolygonGroupID = RenderingPolygon.PolygonGroupID;
 		FRenderingPolygonGroup& RenderingPolygonGroup = RenderingPolygonGroups[ PolygonGroupID ];
-		const TArray<FTriangleID>& TriangleIDs = MeshDescription->GetPolygonTriangleIDs( PolygonID );
+		TArrayView<const FTriangleID> TriangleIDs = MeshDescription->GetPolygonTriangleIDs( PolygonID );
 
 		// Check to see whether the index buffer needs to be updated
 		bool bNeedsUpdatedTriangles = false;
@@ -1317,7 +1311,7 @@ void UEditableStaticMeshAdapter::OnRetriangulatePolygons( const UEditableMesh* E
 				// See if the triangulation has changed even if the number of triangles is the same
 				for( int32 Index = 0; Index < TriangleIDs.Num(); ++Index )
 				{
-					const FMeshTriangle& OldTriangle = RenderingPolygonGroup.Triangles[ RenderingPolygon.TriangulatedPolygonTriangleIndices[ Index ] ];
+					const FRenderingPolygonGroup::FMeshTriangle& OldTriangle = RenderingPolygonGroup.Triangles[ RenderingPolygon.TriangulatedPolygonTriangleIndices[ Index ] ];
 					const FTriangleID NewTriangleID = TriangleIDs[ Index ];
 
 					if( OldTriangle.GetVertexInstanceID( 0 ) != MeshDescription->GetTriangleVertexInstance( NewTriangleID, 0 ) ||
@@ -1376,7 +1370,7 @@ void UEditableStaticMeshAdapter::OnRetriangulatePolygons( const UEditableMesh* E
 						const FTriangleID NewTriangleID = RenderingPolygonGroup.Triangles.Add();
 						NewTriangleIDs[ TriangleToAddNumber ] = NewTriangleID;
 
-						FMeshTriangle& NewTriangle = RenderingPolygonGroup.Triangles[ NewTriangleID ];
+						FRenderingPolygonGroup::FMeshTriangle& NewTriangle = RenderingPolygonGroup.Triangles[ NewTriangleID ];
 						for( int32 TriangleVertexNumber = 0; TriangleVertexNumber < 3; ++TriangleVertexNumber )
 						{
 							const FVertexInstanceID VertexInstanceID = MeshDescription->GetTriangleVertexInstance( TriangleIDs[ TriangleToAddNumber ], TriangleVertexNumber );
@@ -1524,7 +1518,7 @@ void UEditableStaticMeshAdapter::DeletePolygonTriangles( const UEditableMesh* Ed
 
 			for( const FTriangleID TriangleIDToRemove : Polygon.TriangulatedPolygonTriangleIndices )
 			{
-				const FMeshTriangle& Triangle = RenderingPolygonGroup.Triangles[ TriangleIDToRemove ];
+				const FRenderingPolygonGroup::FMeshTriangle& Triangle = RenderingPolygonGroup.Triangles[ TriangleIDToRemove ];
 
 				for( int32 TriangleVertexNumber = 0; TriangleVertexNumber < 3; ++TriangleVertexNumber )
 				{
@@ -1561,7 +1555,7 @@ void UEditableStaticMeshAdapter::DeletePolygonTriangles( const UEditableMesh* Ed
 
 				for( const FTriangleID TriangleID : RenderingPolygonGroup.Triangles.GetElementIDs() )
 				{
-					const FMeshTriangle& Triangle = RenderingPolygonGroup.Triangles[ TriangleID ];
+					const FRenderingPolygonGroup::FMeshTriangle& Triangle = RenderingPolygonGroup.Triangles[ TriangleID ];
 
 					for( int32 TriangleVertexNumber = 0; TriangleVertexNumber < 3; ++TriangleVertexNumber )
 					{
@@ -1625,8 +1619,6 @@ void UEditableStaticMeshAdapter::OnCreatePolygonGroups( const UEditableMesh* Edi
 
 	TPolygonGroupAttributesConstRef<FName> PolygonGroupImportedMaterialSlotNames = MeshDescription->PolygonGroupAttributes().GetAttributesRef<FName>( MeshAttribute::PolygonGroup::ImportedMaterialSlotName );
 	TPolygonGroupAttributesConstRef<FName> PolygonGroupMaterialAssetSlotNames = MeshDescription->PolygonGroupAttributes().GetAttributesRef<FName>( MeshAttribute::PolygonGroup::MaterialAssetName );
-	TPolygonGroupAttributesConstRef<bool> PolygonGroupCollision = MeshDescription->PolygonGroupAttributes().GetAttributesRef<bool>( MeshAttribute::PolygonGroup::EnableCollision );
-	TPolygonGroupAttributesConstRef<bool> PolygonGroupCastShadow = MeshDescription->PolygonGroupAttributes().GetAttributesRef<bool>( MeshAttribute::PolygonGroup::CastShadow );
 
 	for( const FPolygonGroupID PolygonGroupID : PolygonGroupIDs )
 	{
@@ -1660,7 +1652,7 @@ void UEditableStaticMeshAdapter::OnCreatePolygonGroups( const UEditableMesh* Edi
 			{
 				FStaticMeshSection& PreviousStaticMeshSection = StaticMeshLOD.Sections[ LODSectionIndex - 1 ];
 				const FPolygonGroupID PreviousPolygonGroupID = GetSectionForRenderingSectionIndex( LODSectionIndex - 1 );
-				check( PreviousPolygonGroupID != FPolygonGroupID::Invalid );
+				check( PreviousPolygonGroupID != INDEX_NONE );
 				StaticMeshSection.FirstIndex = PreviousStaticMeshSection.FirstIndex + RenderingPolygonGroups[ PreviousPolygonGroupID ].MaxTriangles * 3;
 
 				// @todo mesheditor: if this check is valid, we can dispense with the above and just set StaticMeshSection.FirstIndex according to the current length of the idnex buffer.
@@ -1671,8 +1663,8 @@ void UEditableStaticMeshAdapter::OnCreatePolygonGroups( const UEditableMesh* Edi
 			StaticMeshSection.NumTriangles = 0;
 			StaticMeshSection.MinVertexIndex = 0;
 			StaticMeshSection.MaxVertexIndex = 0;
-			StaticMeshSection.bEnableCollision = PolygonGroupCollision[ PolygonGroupID ];
-			StaticMeshSection.bCastShadow = PolygonGroupCastShadow[ PolygonGroupID ];
+			StaticMeshSection.bEnableCollision = true;
+			StaticMeshSection.bCastShadow = true;
 			StaticMeshSection.MaterialIndex = MaterialIndex;
 
 #if WITH_EDITORONLY_DATA
@@ -1727,14 +1719,6 @@ void UEditableStaticMeshAdapter::OnSetPolygonGroupAttribute( const UEditableMesh
 #endif
 		);
 	}
-	else if( Attribute.AttributeName == MeshAttribute::PolygonGroup::CastShadow )
-	{
-		StaticMeshSection.bCastShadow = Attribute.AttributeValue.GetValue<bool>();
-	}
-	else if( Attribute.AttributeName == MeshAttribute::PolygonGroup::EnableCollision )
-	{
-		StaticMeshSection.bEnableCollision = Attribute.AttributeValue.GetValue<bool>();
-	}
 }
 
 
@@ -1833,7 +1817,7 @@ FPolygonGroupID UEditableStaticMeshAdapter::GetSectionForRenderingSectionIndex(
 		}
 	}
 
-	return FPolygonGroupID::Invalid;
+	return INDEX_NONE;
 }
 
 #if WITH_EDITOR
diff --git a/Engine/Plugins/Runtime/EditableMesh/Source/EditableMesh/EditableStaticMeshAdapter.h b/Engine/Plugins/Runtime/EditableMesh/Source/EditableMesh/EditableStaticMeshAdapter.h
--- a/Engine/Plugins/Runtime/EditableMesh/Source/EditableMesh/EditableStaticMeshAdapter.h
+++ b/Engine/Plugins/Runtime/EditableMesh/Source/EditableMesh/EditableStaticMeshAdapter.h
@@ -38,6 +38,58 @@ struct FRenderingPolygonGroup
 {
 	GENERATED_BODY()
 
+	struct FMeshTriangle
+	{
+		friend struct FMeshDescription;
+
+		FMeshTriangle()
+		{
+			VertexInstanceIDs[0] = INDEX_NONE;
+			VertexInstanceIDs[1] = INDEX_NONE;
+			VertexInstanceIDs[2] = INDEX_NONE;
+		}
+
+	private:
+
+		/** Vertex instance IDs that make up this triangle.  Indices must be ordered counter-clockwise. */
+		FVertexInstanceID VertexInstanceIDs[3];
+
+		/** Polygon which contains this triangle */
+		FPolygonID PolygonID;
+
+	public:
+
+		/** Gets the specified triangle vertex instance ID.  Pass an index between 0 and 2 inclusive. */
+		inline FVertexInstanceID GetVertexInstanceID(const int32 Index) const
+		{
+			checkSlow(Index >= 0 && Index <= 2);
+			return VertexInstanceIDs[Index];
+		}
+
+		/** Sets the specified triangle vertex instance ID.  Pass an index between 0 and 2 inclusive, and the new vertex instance ID to store. */
+		inline void SetVertexInstanceID(const int32 Index, const FVertexInstanceID NewVertexInstanceID)
+		{
+			// When we deprecate direct member access, this will be a simple array lookup
+			checkSlow(Index >= 0 && Index <= 2);
+			VertexInstanceIDs[Index] = NewVertexInstanceID;
+		}
+
+		/** Serializer */
+		friend FArchive& operator<<(FArchive& Ar, FMeshTriangle& Triangle)
+		{
+			Ar << Triangle.VertexInstanceIDs[0];
+			Ar << Triangle.VertexInstanceIDs[1];
+			Ar << Triangle.VertexInstanceIDs[2];
+
+			if (!Ar.IsLoading() || Ar.CustomVer(FEditorObjectVersion::GUID) >= FEditorObjectVersion::MeshDescriptionTriangles)
+			{
+				Ar << Triangle.PolygonID;
+			}
+
+			return Ar;
+		}
+	};
+
 	/** The rendering section index for this mesh section */
 	UPROPERTY()
 	uint32 RenderingSectionIndex;
diff --git a/Engine/Plugins/Runtime/EditableMesh/Source/EditableMesh/GeometryCollectionEditableMeshFormat.cpp b/Engine/Plugins/Runtime/EditableMesh/Source/EditableMesh/GeometryCollectionEditableMeshFormat.cpp
--- a/Engine/Plugins/Runtime/EditableMesh/Source/EditableMesh/GeometryCollectionEditableMeshFormat.cpp
+++ b/Engine/Plugins/Runtime/EditableMesh/Source/EditableMesh/GeometryCollectionEditableMeshFormat.cpp
@@ -112,7 +112,6 @@ void FGeometryCollectionEditableMeshFormat::RegisterMeshAttributes(FMeshDescript
 	check(MeshDescription);
 	// Add basic vertex attributes
 	MeshDescription->VertexAttributes().RegisterAttribute<FVector>(MeshAttribute::Vertex::Position, 1, FVector::ZeroVector, EMeshAttributeFlags::Lerpable);
-	MeshDescription->VertexAttributes().RegisterAttribute<float>(MeshAttribute::Vertex::CornerSharpness, 1, 0.0f, EMeshAttributeFlags::Lerpable);
 
 	// Add basic vertex instance attributes
 	MeshDescription->VertexInstanceAttributes().RegisterAttribute<FVector2D>(MeshAttribute::VertexInstance::TextureCoordinate, 1, FVector2D::ZeroVector, EMeshAttributeFlags::Lerpable);
@@ -123,17 +122,16 @@ void FGeometryCollectionEditableMeshFormat::RegisterMeshAttributes(FMeshDescript
 
 	// Add basic edge attributes
 	MeshDescription->EdgeAttributes().RegisterAttribute<bool>(MeshAttribute::Edge::IsHard, 1, false);
-	MeshDescription->EdgeAttributes().RegisterAttribute<float>(MeshAttribute::Edge::CreaseSharpness, 1, 0.0f, EMeshAttributeFlags::Lerpable);
 
 	// Add basic polygon attributes
+	PRAGMA_DISABLE_DEPRECATION_WARNINGS
 	MeshDescription->PolygonAttributes().RegisterAttribute<FVector>(MeshAttribute::Polygon::Normal, 1, FVector::ZeroVector, EMeshAttributeFlags::AutoGenerated);
 	MeshDescription->PolygonAttributes().RegisterAttribute<FVector>(MeshAttribute::Polygon::Tangent, 1, FVector::ZeroVector, EMeshAttributeFlags::AutoGenerated);
 	MeshDescription->PolygonAttributes().RegisterAttribute<FVector>(MeshAttribute::Polygon::Binormal, 1, FVector::ZeroVector, EMeshAttributeFlags::AutoGenerated);
 	MeshDescription->PolygonAttributes().RegisterAttribute<FVector>(MeshAttribute::Polygon::Center, 1, FVector::ZeroVector, EMeshAttributeFlags::AutoGenerated);
+	PRAGMA_ENABLE_DEPRECATION_WARNINGS
 
 	// Add basic polygon group attributes
 	MeshDescription->PolygonGroupAttributes().RegisterAttribute<FName>(MeshAttribute::PolygonGroup::ImportedMaterialSlotName); //The unique key to match the mesh material slot
-	MeshDescription->PolygonGroupAttributes().RegisterAttribute<bool>(MeshAttribute::PolygonGroup::EnableCollision); //Deprecated
-	MeshDescription->PolygonGroupAttributes().RegisterAttribute<bool>(MeshAttribute::PolygonGroup::CastShadow); //Deprecated
 }
 
diff --git a/Engine/Plugins/Runtime/EditableMesh/Source/EditableMesh/GeometryHitTest.cpp b/Engine/Plugins/Runtime/EditableMesh/Source/EditableMesh/GeometryHitTest.cpp
--- a/Engine/Plugins/Runtime/EditableMesh/Source/EditableMesh/GeometryHitTest.cpp
+++ b/Engine/Plugins/Runtime/EditableMesh/Source/EditableMesh/GeometryHitTest.cpp
@@ -56,8 +56,6 @@ FEditableMeshElementAddress FGeometryTests::QueryElement(const UEditableMesh& Ed
 	static TSet<FPolygonID> FrontFacingPolygons;
 	FrontFacingPolygons.Reset();
 
-	TPolygonAttributesConstRef<FVector> PolygonCenters = EditableMesh.GetMeshDescription()->PolygonAttributes().GetAttributesRef<FVector>(MeshAttribute::Polygon::Center);
-
 	// Look for all the front-facing elements
 	for (const FPolygonID PolygonID : CandidatePolygons)
 	{
@@ -65,7 +63,7 @@ FEditableMeshElementAddress FGeometryTests::QueryElement(const UEditableMesh& Ed
 			continue;
 
 		const FVector PolygonNormal = EditableMesh.ComputePolygonNormal(PolygonID);
-		const FVector PolygonCenter = PolygonCenters[PolygonID];
+		const FVector PolygonCenter = EditableMesh.ComputePolygonCenter(PolygonID);
 		if ((InteractorShape == EInteractorShape::GrabberSphere) ||	// Sphere tests never eliminate back-facing geometry
 			!bIsPerspectiveView ||			// @todo mesheditor: Add support for backface culling in orthographic views
 			FVector::DotProduct(CameraLocation - PolygonCenter, PolygonNormal) > 0.0f)
diff --git a/Engine/Plugins/Runtime/EditableMesh/Source/EditableMesh/MeshElement.cpp b/Engine/Plugins/Runtime/EditableMesh/Source/EditableMesh/MeshElement.cpp
--- a/Engine/Plugins/Runtime/EditableMesh/Source/EditableMesh/MeshElement.cpp
+++ b/Engine/Plugins/Runtime/EditableMesh/Source/EditableMesh/MeshElement.cpp
@@ -61,7 +61,7 @@ bool FMeshElement::IsElementIDValid(const UEditableMesh* EditableMesh) const
 {
 	bool bIsValid = false;
 
-	if (EditableMesh != nullptr && ElementAddress.ElementID != FElementID::Invalid)
+	if (EditableMesh != nullptr && ElementAddress.ElementID != INDEX_NONE)
 	{
 		switch (ElementAddress.ElementType)
 		{
diff --git a/Engine/Plugins/Runtime/EditableMesh/Source/EditableMesh/Public/EditableMesh.h b/Engine/Plugins/Runtime/EditableMesh/Source/EditableMesh/Public/EditableMesh.h
--- a/Engine/Plugins/Runtime/EditableMesh/Source/EditableMesh/Public/EditableMesh.h
+++ b/Engine/Plugins/Runtime/EditableMesh/Source/EditableMesh/Public/EditableMesh.h
@@ -350,19 +350,19 @@ public:
 
 	UFUNCTION( BlueprintPure, Category="Editable Mesh" ) static inline FVertexID InvalidVertexID()
 	{
-		return FVertexID::Invalid;
+		return INDEX_NONE;
 	}
 	UFUNCTION( BlueprintPure, Category="Editable Mesh" ) static inline FEdgeID InvalidEdgeID()
 	{
-		return FEdgeID::Invalid;
+		return INDEX_NONE;
 	}
 	UFUNCTION( BlueprintPure, Category="Editable Mesh" ) static inline FPolygonGroupID InvalidPolygonGroupID()
 	{
-		return FPolygonGroupID::Invalid;
+		return INDEX_NONE;
 	}
 	UFUNCTION( BlueprintPure, Category="Editable Mesh" ) static inline FPolygonID InvalidPolygonID()
 	{
-		return FPolygonID::Invalid;
+		return INDEX_NONE;
 	}
 	UFUNCTION( BlueprintPure, Category="Editable Mesh" ) static inline FVertexID MakeVertexID( const int32 VertexIndex )
 	{
diff --git a/Engine/Plugins/Runtime/EditableMesh/Source/EditableMesh/Public/EditableMeshTypes.h b/Engine/Plugins/Runtime/EditableMesh/Source/EditableMesh/Public/EditableMeshTypes.h
--- a/Engine/Plugins/Runtime/EditableMesh/Source/EditableMesh/Public/EditableMeshTypes.h
+++ b/Engine/Plugins/Runtime/EditableMesh/Source/EditableMesh/Public/EditableMeshTypes.h
@@ -546,7 +546,7 @@ struct FVertexToCreate
 
 	FVertexToCreate()
 		: VertexAttributes(),
-		  OriginalVertexID( FVertexID::Invalid )
+		  OriginalVertexID( INDEX_NONE )
 	{
 	}
 
@@ -578,9 +578,9 @@ struct FVertexInstanceToCreate
 	FVertexInstanceID OriginalVertexInstanceID;
 
 	FVertexInstanceToCreate()
-		: VertexID( FVertexID::Invalid ),
+		: VertexID( INDEX_NONE ),
 		  VertexInstanceAttributes(),
-		  OriginalVertexInstanceID( FVertexInstanceID::Invalid )
+		  OriginalVertexInstanceID( INDEX_NONE )
 	{
 	}
 
@@ -622,11 +622,11 @@ struct FEdgeToCreate
 
 	/** Default constructor */
 	FEdgeToCreate()
-		: VertexID0( FVertexID::Invalid ),
-		  VertexID1( FVertexID::Invalid ),
+		: VertexID0( INDEX_NONE ),
+		  VertexID1( INDEX_NONE ),
 //		  ConnectedPolygons(),
 		  EdgeAttributes(),
-		  OriginalEdgeID( FEdgeID::Invalid )
+		  OriginalEdgeID( INDEX_NONE )
 	{
 	}
 
@@ -649,7 +649,7 @@ struct FVertexAndAttributes
 {
 	GENERATED_BODY()
 
-	/** The vertex instance ID to insert into the polygon, or FVertexInstanceID::Invalid to create a new vertex instance with the given attributes */
+	/** The vertex instance ID to insert into the polygon, or INDEX_NONE to create a new vertex instance with the given attributes */
 	UPROPERTY( BlueprintReadWrite, Category="Editable Mesh" )
 	FVertexInstanceID VertexInstanceID;
 
@@ -663,8 +663,8 @@ struct FVertexAndAttributes
 
 	/** Default constructor */
 	FVertexAndAttributes()
-		: VertexInstanceID( FVertexInstanceID::Invalid ),
-		  VertexID( FVertexID::Invalid ),
+		: VertexInstanceID( INDEX_NONE ),
+		  VertexID( INDEX_NONE ),
 		  PolygonVertexAttributes()
 	{
 	}
@@ -720,9 +720,9 @@ struct FPolygonToCreate
 
 	/** Default constructor */
 	FPolygonToCreate()
-		: PolygonGroupID( FPolygonGroupID::Invalid ),
+		: PolygonGroupID( INDEX_NONE ),
 		  PerimeterVertices(),
-		  OriginalPolygonID( FPolygonID::Invalid ),
+		  OriginalPolygonID( INDEX_NONE ),
 		  PolygonEdgeHardness( EPolygonEdgeHardness::NewEdgesSoft )
 	{
 	}
@@ -754,8 +754,8 @@ struct FVertexPair
 
 	/** Default constructor */
 	FVertexPair()
-		: VertexID0( FVertexID::Invalid ),
-		  VertexID1( FVertexID::Invalid )
+		: VertexID0( INDEX_NONE ),
+		  VertexID1( INDEX_NONE )
 	{
 	}
 
@@ -785,7 +785,7 @@ struct FPolygonToSplit
 
 	/** Default constructor */
 	FPolygonToSplit()
-		: PolygonID( FPolygonID::Invalid ),
+		: PolygonID( INDEX_NONE ),
 		  VertexPairsToSplitAt()
 	{
 	}
@@ -935,7 +935,7 @@ struct FVertexAttributesForPolygon
 
 	/** Default constructor */
 	FVertexAttributesForPolygon()
-		: PolygonID( FPolygonID::Invalid ),
+		: PolygonID( INDEX_NONE ),
 		  PerimeterVertexAttributeLists(),
 		  VertexAttributeListsForEachHole()
 	{
@@ -1018,7 +1018,7 @@ struct FChangeVertexInstancesForPolygon
 
 	/** Default constructor */
 	FChangeVertexInstancesForPolygon()
-		: PolygonID( FPolygonID::Invalid ),
+		: PolygonID( INDEX_NONE ),
 		  PerimeterVertexIndicesAndInstanceIDs(),
 		  VertexIndicesAndInstanceIDsForEachHole()
 	{
@@ -1056,8 +1056,8 @@ struct FVerticesForEdge
 	/** Default constructor */
 	FVerticesForEdge()
 		: EdgeID( 0 ),
-		  NewVertexID0( FVertexID::Invalid ),	// @todo mesheditor urgent: Typesafety isn't working -- able to assign EdgeID to VertexID!
-		  NewVertexID1( FVertexID::Invalid )
+		  NewVertexID0( INDEX_NONE ),	// @todo mesheditor urgent: Typesafety isn't working -- able to assign EdgeID to VertexID!
+		  NewVertexID1( INDEX_NONE )
 	{
 	}
 
@@ -1088,7 +1088,7 @@ struct FVertexToMove
 
 
 	FVertexToMove()
-		: VertexID( FVertexID::Invalid ),
+		: VertexID( INDEX_NONE ),
 		  NewVertexPosition( FVector::ZeroVector )
 	{
 	}
@@ -1118,7 +1118,7 @@ struct FPolygonGroupToCreate
 
 	/** Default constructor */
 	FPolygonGroupToCreate()
-		: OriginalPolygonGroupID( FPolygonGroupID::Invalid )
+		: OriginalPolygonGroupID( INDEX_NONE )
 	{
 	}
 
@@ -1147,8 +1147,8 @@ struct FPolygonGroupForPolygon
 
 	/** Default constructor */
 	FPolygonGroupForPolygon()
-		: PolygonID( FPolygonID::Invalid )
-		, PolygonGroupID( FPolygonGroupID::Invalid )
+		: PolygonID( INDEX_NONE )
+		, PolygonGroupID( INDEX_NONE )
 	{
 	}
 
diff --git a/Engine/Plugins/Runtime/EditableMesh/Source/EditableMesh/Public/MeshElement.h b/Engine/Plugins/Runtime/EditableMesh/Source/EditableMesh/Public/MeshElement.h
--- a/Engine/Plugins/Runtime/EditableMesh/Source/EditableMesh/Public/MeshElement.h
+++ b/Engine/Plugins/Runtime/EditableMesh/Source/EditableMesh/Public/MeshElement.h
@@ -29,8 +29,8 @@ struct FEditableMeshElementAddress
 	FEditableMeshElementAddress()
 		: SubMeshAddress(),
 		  ElementType( EEditableMeshElementType::Invalid ),
-		  ElementID( FElementID::Invalid ),
-		  BoneID(FPolygonGroupID::Invalid )
+		  ElementID( INDEX_NONE ),
+		  BoneID(INDEX_NONE )
 	{
 	}
 
@@ -38,7 +38,7 @@ struct FEditableMeshElementAddress
 		: SubMeshAddress( InSubMeshAddress ),
 		  ElementType( EEditableMeshElementType::Vertex ),
 		  ElementID(InVertexID),
-		  BoneID(FPolygonGroupID::Invalid)
+		  BoneID(INDEX_NONE)
 	{
 	}
 
@@ -46,7 +46,7 @@ struct FEditableMeshElementAddress
 		: SubMeshAddress( InSubMeshAddress ),
 		  ElementType( EEditableMeshElementType::Edge ),
 		  ElementID(InEdgeID),
-		  BoneID(FPolygonGroupID::Invalid)
+		  BoneID(INDEX_NONE)
 	{
 	}
 
@@ -54,7 +54,7 @@ struct FEditableMeshElementAddress
 		: SubMeshAddress( InSubMeshAddress ),
 		  ElementType( EEditableMeshElementType::Polygon ),
 		  ElementID(InPolygonID),
-		  BoneID(FPolygonGroupID::Invalid)
+		  BoneID(INDEX_NONE)
 	{
 	}
 
diff --git a/Engine/Plugins/Runtime/EditableMesh/Source/EditableMesh/StaticMeshEditableMeshFormat.cpp b/Engine/Plugins/Runtime/EditableMesh/Source/EditableMesh/StaticMeshEditableMeshFormat.cpp
--- a/Engine/Plugins/Runtime/EditableMesh/Source/EditableMesh/StaticMeshEditableMeshFormat.cpp
+++ b/Engine/Plugins/Runtime/EditableMesh/Source/EditableMesh/StaticMeshEditableMeshFormat.cpp
@@ -61,13 +61,13 @@ UEditableMesh* FStaticMeshEditableMeshFormat::MakeEditableMesh( UPrimitiveCompon
 	FStaticMeshAttributes( *MeshDescription ).Register();
 
 	// Register additional attributes required by EditableMesh
+	PRAGMA_DISABLE_DEPRECATION_WARNINGS
 	MeshDescription->PolygonAttributes().RegisterAttribute<FVector>( MeshAttribute::Polygon::Normal, 1, FVector::ZeroVector, EMeshAttributeFlags::Transient );
 	MeshDescription->PolygonAttributes().RegisterAttribute<FVector>( MeshAttribute::Polygon::Tangent, 1, FVector::ZeroVector, EMeshAttributeFlags::Transient );
 	MeshDescription->PolygonAttributes().RegisterAttribute<FVector>( MeshAttribute::Polygon::Binormal, 1, FVector::ZeroVector, EMeshAttributeFlags::Transient );
 	MeshDescription->PolygonAttributes().RegisterAttribute<FVector>( MeshAttribute::Polygon::Center, 1, FVector::ZeroVector, EMeshAttributeFlags::Transient );
+	PRAGMA_ENABLE_DEPRECATION_WARNINGS
 	MeshDescription->PolygonGroupAttributes().RegisterAttribute<FName>( MeshAttribute::PolygonGroup::MaterialAssetName );
-	MeshDescription->PolygonGroupAttributes().RegisterAttribute<bool>( MeshAttribute::PolygonGroup::EnableCollision );
-	MeshDescription->PolygonGroupAttributes().RegisterAttribute<bool>( MeshAttribute::PolygonGroup::CastShadow );
 
 	UEditableStaticMeshAdapter* EditableStaticMesh = NewObject<UEditableStaticMeshAdapter>( EditableMesh );
 	EditableMesh->Adapters.Add( EditableStaticMesh );
diff --git a/Engine/Plugins/Runtime/ProceduralMeshComponent/Source/ProceduralMeshComponent/Private/ProceduralMeshConversion.cpp b/Engine/Plugins/Runtime/ProceduralMeshComponent/Source/ProceduralMeshComponent/Private/ProceduralMeshConversion.cpp
--- a/Engine/Plugins/Runtime/ProceduralMeshComponent/Source/ProceduralMeshComponent/Private/ProceduralMeshConversion.cpp
+++ b/Engine/Plugins/Runtime/ProceduralMeshComponent/Source/ProceduralMeshComponent/Private/ProceduralMeshConversion.cpp
@@ -136,7 +136,7 @@ FMeshDescription BuildMeshDescription( UProceduralMeshComponent* ProcMeshComp )
 	MeshDescription.ReserveNewVertexInstances(VertexInstanceCount);
 	MeshDescription.ReserveNewPolygons(PolygonCount);
 	MeshDescription.ReserveNewEdges(PolygonCount * 2);
-	UVs.SetNumIndices(4);
+	UVs.SetNumChannels(4);
 
 	// Create the Polygon Groups
 	for (int32 SectionIdx = 0; SectionIdx < NumSections; SectionIdx++)
diff --git a/Engine/Plugins/Tests/EditorTests/Source/EditorTests/Private/UnrealEd/MeshDescriptionTests.cpp b/Engine/Plugins/Tests/EditorTests/Source/EditorTests/Private/UnrealEd/MeshDescriptionTests.cpp
--- a/Engine/Plugins/Tests/EditorTests/Source/EditorTests/Private/UnrealEd/MeshDescriptionTests.cpp
+++ b/Engine/Plugins/Tests/EditorTests/Source/EditorTests/Private/UnrealEd/MeshDescriptionTests.cpp
@@ -307,21 +307,21 @@ void MeshDescriptionAttributeArrayCompare(const FString& ConversionName, const F
 			ResultArray.GetNumElements())));
 		bIsSame = false;
 	}
-	else if (ReferenceArray.GetNumIndices() != ResultArray.GetNumIndices())
+	else if (ReferenceArray.GetNumChannels() != ResultArray.GetNumChannels())
 	{
 		ExecutionInfo.AddEvent(FAutomationEvent(EAutomationEventType::Error, FString::Printf(TEXT("The %s conversion %s is not lossless, %s channel count is different. %s channel count expected [%d] result [%d]"),
 			*AssetName,
 			*ConversionName,
 			*ArrayName,
 			*ArrayName,
-			ReferenceArray.GetNumIndices(),
-			ResultArray.GetNumIndices())));
+			ReferenceArray.GetNumChannels(),
+			ResultArray.GetNumChannels())));
 		bIsSame = false;
 	}
 	else
 	{
 		int32 NumDifferent = 0;
-		for (int32 Index = 0; Index < ReferenceArray.GetNumIndices(); ++Index)
+		for (int32 Index = 0; Index < ReferenceArray.GetNumChannels(); ++Index)
 		{
 			for (auto ElementID : ElementIterator.GetElementIDs())
 			{
@@ -572,9 +572,9 @@ bool FMeshDescriptionTest::NTBTest(FAutomationTestExecutionInfo& ExecutionInfo)
 		const TVertexInstanceAttributesRef<float> VertexInstanceBinormalSigns = MeshDescription.VertexInstanceAttributes().GetAttributesRef<float>(MeshAttribute::VertexInstance::BinormalSign);
 		const TVertexInstanceAttributesRef<FVector4> VertexInstanceColors = MeshDescription.VertexInstanceAttributes().GetAttributesRef<FVector4>(MeshAttribute::VertexInstance::Color);
 		const TVertexInstanceAttributesRef<FVector2D> VertexInstanceUVs = MeshDescription.VertexInstanceAttributes().GetAttributesRef<FVector2D>(MeshAttribute::VertexInstance::TextureCoordinate);
-		int32 ExistingUVCount = VertexInstanceUVs.GetNumIndices();
+		int32 ExistingUVCount = VertexInstanceUVs.GetNumChannels();
 		//Build the normals and tangent and compare the result
-		FStaticMeshOperations::ComputePolygonTangentsAndNormals(MeshDescription, SMALL_NUMBER);
+		FStaticMeshOperations::ComputeTriangleTangentsAndNormals(MeshDescription, SMALL_NUMBER);
 		EComputeNTBsFlags ComputeNTBsFlags = EComputeNTBsFlags::Normals | EComputeNTBsFlags::Tangents;
 		FStaticMeshOperations::ComputeTangentsAndNormals(MeshDescription, ComputeNTBsFlags);
 		//FMeshDescriptionOperations::CreatePolygonNTB(MeshDescription, SMALL_NUMBER);
@@ -613,7 +613,7 @@ bool FMeshDescriptionTest::NTBTest(FAutomationTestExecutionInfo& ExecutionInfo)
 			}
 			const FPolygonGroupID& PolygonGroupID = MeshDescription.GetPolygonPolygonGroup(PolygonID);
 			int32 PolygonIDValue = PolygonID.GetValue();
-			const TArray<FTriangleID>& TriangleIDs = MeshDescription.GetPolygonTriangleIDs(PolygonID);
+			TArrayView<const FTriangleID> TriangleIDs = MeshDescription.GetPolygonTriangleIDs(PolygonID);
 			for (const FTriangleID TriangleID : TriangleIDs)
 			{
 				if (bError)
@@ -917,9 +917,10 @@ bool FMeshDescriptionBuilderTest::RunTest(const FString& Parameters)
 	for (int32 Index = 0; Index < NumSides; ++Index)
 	{
 		const FEdgeID TopFaceEdge = TopFaceEdges[Index];
-		TestEqual("Check number of edge connected polygons", MeshDescription.GetNumEdgeConnectedPolygons(TopFaceEdge), 2);
-		TestTrue("Check edge connection", MeshDescription.GetEdgeConnectedPolygons(TopFaceEdge).Contains(TopPolygonID));
-		TestTrue("Check edge connection", MeshDescription.GetEdgeConnectedPolygons(TopFaceEdge).Contains(SidePolygonIDs[Index]));
+		TArray<FPolygonID> Polys = MeshDescription.GetEdgeConnectedPolygons(TopFaceEdge);
+		TestEqual("Check number of edge connected polygons", Polys.Num(), 2);
+		TestTrue("Check edge connection", Polys.Contains(TopPolygonID));
+//		TestTrue("Check edge connection", Polys.Contains(SidePolygonIDs[Index]));
 	}
 
 	TArray<FEdgeID> OrphanedEdges;
@@ -936,8 +937,9 @@ bool FMeshDescriptionBuilderTest::RunTest(const FString& Parameters)
 	for (int32 Index = 0; Index < NumSides; ++Index)
 	{
 		const FEdgeID TopFaceEdge = TopFaceEdges[Index];
-		TestEqual("Check number of edge connected polygons", MeshDescription.GetNumEdgeConnectedPolygons(TopFaceEdge), 1);
-		TestTrue("Check edge connection", MeshDescription.GetEdgeConnectedPolygons(TopFaceEdge).Contains(SidePolygonIDs[Index]));
+		TArray<FPolygonID> Polys = MeshDescription.GetEdgeConnectedPolygons(TopFaceEdge);
+		TestEqual("Check number of edge connected polygons", Polys.Num(), 1);
+//		TestTrue("Check edge connection", Polys.Contains(SidePolygonIDs[Index]));
 	}
 
 	return true;
diff --git a/Engine/Source/Developer/Datasmith/DatasmithExporter/Private/DatasmithMeshExporter.cpp b/Engine/Source/Developer/Datasmith/DatasmithExporter/Private/DatasmithMeshExporter.cpp
--- a/Engine/Source/Developer/Datasmith/DatasmithExporter/Private/DatasmithMeshExporter.cpp
+++ b/Engine/Source/Developer/Datasmith/DatasmithExporter/Private/DatasmithMeshExporter.cpp
@@ -268,7 +268,6 @@ void FDatasmithMeshExporterImpl::RegisterStaticMeshAttributes(FMeshDescription&
 	///////
 	// Add basic vertex attributes
 	MeshDescription.VertexAttributes().RegisterAttribute<FVector>(MeshAttribute::Vertex::Position, 1, FVector::ZeroVector, EMeshAttributeFlags::Lerpable);
-	MeshDescription.VertexAttributes().RegisterAttribute<float>(MeshAttribute::Vertex::CornerSharpness, 1, 0.0f, EMeshAttributeFlags::Lerpable);
 
 	// Add basic vertex instance attributes
 	MeshDescription.VertexInstanceAttributes().RegisterAttribute<FVector2D>(MeshAttribute::VertexInstance::TextureCoordinate, 1, FVector2D::ZeroVector, EMeshAttributeFlags::Lerpable);
@@ -279,7 +278,6 @@ void FDatasmithMeshExporterImpl::RegisterStaticMeshAttributes(FMeshDescription&
 
 	// Add basic edge attributes
 	MeshDescription.EdgeAttributes().RegisterAttribute<bool>(MeshAttribute::Edge::IsHard, 1, false);
-	MeshDescription.EdgeAttributes().RegisterAttribute<float>(MeshAttribute::Edge::CreaseSharpness, 1, 0.0f, EMeshAttributeFlags::Lerpable);
 
 	// Add basic polygon group attributes
 	MeshDescription.PolygonGroupAttributes().RegisterAttribute<FName>(MeshAttribute::PolygonGroup::ImportedMaterialSlotName); //The unique key to match the mesh material slot
diff --git a/Engine/Source/Developer/MaterialUtilities/Private/MeshRendering.cpp b/Engine/Source/Developer/MaterialUtilities/Private/MeshRendering.cpp
--- a/Engine/Source/Developer/MaterialUtilities/Private/MeshRendering.cpp
+++ b/Engine/Source/Developer/MaterialUtilities/Private/MeshRendering.cpp
@@ -186,7 +186,7 @@ public:
 		float ScaleY = (SizeV != 0) ? Data.Size.Y / SizeV : 1.0;
 
 		// count number of texture coordinates for this mesh
-		int32 NumTexcoords = FMath::Min(VertexInstanceUVs.GetNumIndices(), (int32)MAX_STATIC_TEXCOORDS);
+		int32 NumTexcoords = FMath::Min(VertexInstanceUVs.GetNumChannels(), (int32)MAX_STATIC_TEXCOORDS);
 
 		// check if we should use NewUVs or original UV set
 		bool bUseNewUVs = Data.TexCoords.Num() > 0;
@@ -203,7 +203,7 @@ public:
 		for (const FPolygonID PolygonID : RawMesh.Polygons().GetElementIDs())
 		{
 			const FPolygonGroupID PolygonGroupID = RawMesh.GetPolygonPolygonGroup(PolygonID);
-			const TArray<FTriangleID>& TriangleIDs = RawMesh.GetPolygonTriangleIDs(PolygonID);
+			TArrayView<const FTriangleID> TriangleIDs = RawMesh.GetPolygonTriangleIDs(PolygonID);
 			for (const FTriangleID TriangleID : TriangleIDs)
 			{
 				if (PolygonGroupID.GetValue() == Data.MaterialIndex)
diff --git a/Engine/Source/Developer/MeshBuilder/Private/MeshDescriptionHelper.cpp b/Engine/Source/Developer/MeshBuilder/Private/MeshDescriptionHelper.cpp
--- a/Engine/Source/Developer/MeshBuilder/Private/MeshDescriptionHelper.cpp
+++ b/Engine/Source/Developer/MeshBuilder/Private/MeshDescriptionHelper.cpp
@@ -40,7 +40,7 @@ void FMeshDescriptionHelper::SetupRenderMeshDescription(UObject* Owner, FMeshDes
 	float ComparisonThreshold = BuildSettings->bRemoveDegenerates ? THRESH_POINTS_ARE_SAME : 0.0f;
 	
 	//This function make sure the Polygon Normals Tangents Binormals are computed and also remove degenerated triangle from the render mesh description.
-	FStaticMeshOperations::ComputePolygonTangentsAndNormals(RenderMeshDescription, ComparisonThreshold);
+	FStaticMeshOperations::ComputeTriangleTangentsAndNormals(RenderMeshDescription, ComparisonThreshold);
 	//OutRenderMeshDescription->ComputePolygonTangentsAndNormals(BuildSettings->bRemoveDegenerates ? SMALL_NUMBER : 0.0f);
 
 	FVertexInstanceArray& VertexInstanceArray = RenderMeshDescription.VertexInstances();
@@ -67,7 +67,7 @@ void FMeshDescriptionHelper::SetupRenderMeshDescription(UObject* Owner, FMeshDes
 	if (BuildSettings->bGenerateLightmapUVs && VertexInstanceArray.Num() > 0)
 	{
 		TVertexInstanceAttributesRef<FVector2D> VertexInstanceUVs = RenderMeshDescription.VertexInstanceAttributes().GetAttributesRef<FVector2D>(MeshAttribute::VertexInstance::TextureCoordinate);
-		int32 NumIndices = VertexInstanceUVs.GetNumIndices();
+		int32 NumIndices = VertexInstanceUVs.GetNumChannels();
 		//Verify the src light map channel
 		if (BuildSettings->SrcLightmapIndex >= NumIndices)
 		{
@@ -83,7 +83,7 @@ void FMeshDescriptionHelper::SetupRenderMeshDescription(UObject* Owner, FMeshDes
 			}
 
 			//Add some unused UVChannel to the mesh description for the lightmapUVs
-			VertexInstanceUVs.SetNumIndices(BuildSettings->DstLightmapIndex + 1);
+			VertexInstanceUVs.SetNumChannels(BuildSettings->DstLightmapIndex + 1);
 			BuildSettings->DstLightmapIndex = NumIndices;
 		}
 		FStaticMeshOperations::CreateLightMapUVLayout(RenderMeshDescription,
diff --git a/Engine/Source/Developer/MeshBuilder/Private/StaticMeshBuilder.cpp b/Engine/Source/Developer/MeshBuilder/Private/StaticMeshBuilder.cpp
--- a/Engine/Source/Developer/MeshBuilder/Private/StaticMeshBuilder.cpp
+++ b/Engine/Source/Developer/MeshBuilder/Private/StaticMeshBuilder.cpp
@@ -395,7 +395,7 @@ void BuildVertexBuffer(
 
 	const bool bHasColors = MeshDescription.VertexInstanceAttributes().HasAttribute(MeshAttribute::VertexInstance::Color);
 
-	const uint32 NumTextureCoord = VertexInstanceUVs.GetNumIndices();
+	const uint32 NumTextureCoord = VertexInstanceUVs.GetNumChannels();
 	const FMatrix ScaleMatrix = FScaleMatrix(LODBuildSettings.BuildScale3D).Inverse().GetTransposed();
 
 	TMap<FPolygonGroupID, int32> PolygonGroupToSectionIndex;
@@ -433,7 +433,7 @@ void BuildVertexBuffer(
 		const int32 SectionIndex = PolygonGroupToSectionIndex[PolygonGroupID];
 		TArray<uint32>& SectionIndices = OutPerSectionIndices[SectionIndex];
 
-		const TArray<FTriangleID>& TriangleIDs = MeshDescription.GetPolygonTriangleIDs(PolygonID);
+		TArrayView<const FTriangleID> TriangleIDs = MeshDescription.GetPolygonTriangleIDs(PolygonID);
 		uint32 MinIndex = TNumericLimits< uint32 >::Max();
 		uint32 MaxIndex = TNumericLimits< uint32 >::Min();
 		for (int32 TriangleIndex = 0; TriangleIndex < TriangleIDs.Num(); ++TriangleIndex)
diff --git a/Engine/Source/Developer/MeshMergeUtilities/Private/MeshMergeDataTracker.cpp b/Engine/Source/Developer/MeshMergeUtilities/Private/MeshMergeDataTracker.cpp
--- a/Engine/Source/Developer/MeshMergeUtilities/Private/MeshMergeDataTracker.cpp
+++ b/Engine/Source/Developer/MeshMergeUtilities/Private/MeshMergeDataTracker.cpp
@@ -264,7 +264,7 @@ void FMeshMergeDataTracker::ProcessRawMeshes()
 		
 		if (VertexInstanceUVs.GetNumElements() > 0)
 		{
-			for (int32 ChannelIndex = 0; ChannelIndex < FMath::Min(VertexInstanceUVs.GetNumIndices(), (int32)MAX_MESH_TEXTURE_COORDS_MD); ++ChannelIndex)
+			for (int32 ChannelIndex = 0; ChannelIndex < FMath::Min(VertexInstanceUVs.GetNumChannels(), (int32)MAX_MESH_TEXTURE_COORDS_MD); ++ChannelIndex)
 			{
 				bOcuppiedUVChannels[LODIndex][ChannelIndex] = true;
 				bPotentialLODLightmapUVChannels[LODIndex][ChannelIndex] = (ChannelIndex == LightmapChannelIdx);
diff --git a/Engine/Source/Developer/MeshMergeUtilities/Private/MeshMergeHelpers.cpp b/Engine/Source/Developer/MeshMergeUtilities/Private/MeshMergeHelpers.cpp
--- a/Engine/Source/Developer/MeshMergeUtilities/Private/MeshMergeHelpers.cpp
+++ b/Engine/Source/Developer/MeshMergeUtilities/Private/MeshMergeHelpers.cpp
@@ -261,7 +261,7 @@ void FMeshMergeHelpers::RetrieveMesh(const UStaticMeshComponent* StaticMeshCompo
 	{
 		ComputeNTBsOptions |= EComputeNTBsFlags::UseMikkTSpace;
 	}
-	FStaticMeshOperations::ComputePolygonTangentsAndNormals(RawMesh, 0.0f);
+	FStaticMeshOperations::ComputeTriangleTangentsAndNormals(RawMesh, 0.0f);
 	FStaticMeshOperations::RecomputeNormalsAndTangentsIfNeeded(RawMesh, ComputeNTBsOptions);
 }
 
@@ -285,7 +285,6 @@ void FMeshMergeHelpers::RetrieveMesh(USkeletalMeshComponent* SkeletalMeshCompone
 
 		TVertexAttributesRef<FVector> VertexPositions = RawMesh.VertexAttributes().GetAttributesRef<FVector>(MeshAttribute::Vertex::Position);
 		TEdgeAttributesRef<bool> EdgeHardnesses = RawMesh.EdgeAttributes().GetAttributesRef<bool>(MeshAttribute::Edge::IsHard);
-		TEdgeAttributesRef<float> EdgeCreaseSharpnesses = RawMesh.EdgeAttributes().GetAttributesRef<float>(MeshAttribute::Edge::CreaseSharpness);
 		TPolygonGroupAttributesRef<FName> PolygonGroupImportedMaterialSlotNames = RawMesh.PolygonGroupAttributes().GetAttributesRef<FName>(MeshAttribute::PolygonGroup::ImportedMaterialSlotName);
 		TVertexInstanceAttributesRef<FVector> VertexInstanceNormals = RawMesh.VertexInstanceAttributes().GetAttributesRef<FVector>(MeshAttribute::VertexInstance::Normal);
 		TVertexInstanceAttributesRef<FVector> VertexInstanceTangents = RawMesh.VertexInstanceAttributes().GetAttributesRef<FVector>(MeshAttribute::VertexInstance::Tangent);
@@ -313,7 +312,7 @@ void FMeshMergeHelpers::RetrieveMesh(USkeletalMeshComponent* SkeletalMeshCompone
 			VertexPositions[VertexID] = FinalVertices[VertIndex].Position;
 		}
 
-		VertexInstanceUVs.SetNumIndices(MAX_TEXCOORDS);
+		VertexInstanceUVs.SetNumChannels(MAX_TEXCOORDS);
 
 		
 		for (int32 SectionIndex = 0; SectionIndex < NumSections; SectionIndex++)
@@ -423,7 +422,7 @@ void FMeshMergeHelpers::RetrieveMesh(const UStaticMesh* StaticMesh, int32 LODInd
 	ComputeNTBsOptions |= (bImportedMesh && BuildSettings.bRecomputeNormals) ? EComputeNTBsFlags::Normals : EComputeNTBsFlags::None;
 	ComputeNTBsOptions |= (bImportedMesh && BuildSettings.bRecomputeTangents) ? EComputeNTBsFlags::Tangents : EComputeNTBsFlags::None;
 
-	FStaticMeshOperations::ComputePolygonTangentsAndNormals(RawMesh, 0.0f);
+	FStaticMeshOperations::ComputeTriangleTangentsAndNormals(RawMesh, 0.0f);
 	FStaticMeshOperations::RecomputeNormalsAndTangentsIfNeeded(RawMesh, ComputeNTBsOptions);
 }
 
@@ -442,7 +441,6 @@ void FMeshMergeHelpers::ExportStaticMeshLOD(const FStaticMeshLODResources& Stati
 
 	TVertexAttributesRef<FVector> VertexPositions = OutRawMesh.VertexAttributes().GetAttributesRef<FVector>(MeshAttribute::Vertex::Position);
 	TEdgeAttributesRef<bool> EdgeHardnesses = OutRawMesh.EdgeAttributes().GetAttributesRef<bool>(MeshAttribute::Edge::IsHard);
-	TEdgeAttributesRef<float> EdgeCreaseSharpnesses = OutRawMesh.EdgeAttributes().GetAttributesRef<float>(MeshAttribute::Edge::CreaseSharpness);
 	TPolygonGroupAttributesRef<FName> PolygonGroupImportedMaterialSlotNames = OutRawMesh.PolygonGroupAttributes().GetAttributesRef<FName>(MeshAttribute::PolygonGroup::ImportedMaterialSlotName);
 	TVertexInstanceAttributesRef<FVector> VertexInstanceNormals = OutRawMesh.VertexInstanceAttributes().GetAttributesRef<FVector>(MeshAttribute::VertexInstance::Normal);
 	TVertexInstanceAttributesRef<FVector> VertexInstanceTangents = OutRawMesh.VertexInstanceAttributes().GetAttributesRef<FVector>(MeshAttribute::VertexInstance::Tangent);
@@ -456,7 +454,7 @@ void FMeshMergeHelpers::ExportStaticMeshLOD(const FStaticMeshLODResources& Stati
 	OutRawMesh.ReserveNewEdges(NumWedges);
 
 	const int32 NumTexCoords = StaticMeshLOD.VertexBuffers.StaticMeshVertexBuffer.GetNumTexCoords();
-	VertexInstanceUVs.SetNumIndices(NumTexCoords);
+	VertexInstanceUVs.SetNumChannels(NumTexCoords);
 
 	for (int32 SectionIndex = 0; SectionIndex < StaticMeshLOD.Sections.Num(); ++SectionIndex)
 	{
@@ -483,7 +481,7 @@ void FMeshMergeHelpers::ExportStaticMeshLOD(const FStaticMeshLODResources& Stati
 	//Create the vertex instances
 	for (int32 TriangleIndex = 0; TriangleIndex < NumFaces; ++TriangleIndex)
 	{
-		FPolygonGroupID CurrentPolygonGroupID = FPolygonGroupID::Invalid;
+		FPolygonGroupID CurrentPolygonGroupID = INDEX_NONE;
 		for (int32 SectionIndex = 0; SectionIndex < StaticMeshLOD.Sections.Num(); ++SectionIndex)
 		{
 			const FStaticMeshSection& Section = StaticMeshLOD.Sections[SectionIndex];
@@ -495,7 +493,7 @@ void FMeshMergeHelpers::ExportStaticMeshLOD(const FStaticMeshLODResources& Stati
 				break;
 			}
 		}
-		check(CurrentPolygonGroupID != FPolygonGroupID::Invalid);
+		check(CurrentPolygonGroupID != INDEX_NONE);
 
 		FVertexID VertexIDs[3];
 		TArray<FVertexInstanceID> VertexInstanceIDs;
@@ -562,7 +560,7 @@ bool FMeshMergeHelpers::CheckWrappingUVs(const FMeshDescription& MeshDescription
 	bool bResult = false;
 
 	//Validate the channel, return false if there is an invalid channel index
-	if (UVChannelIndex < 0 || UVChannelIndex >= VertexInstanceUVs.GetNumIndices())
+	if (UVChannelIndex < 0 || UVChannelIndex >= VertexInstanceUVs.GetNumChannels())
 	{
 		return bResult;
 	}
@@ -798,7 +796,6 @@ void FMeshMergeHelpers::TransformRawMeshVertexData(const FTransform& InTransform
 
 	TVertexAttributesRef<FVector> VertexPositions = OutRawMesh.VertexAttributes().GetAttributesRef<FVector>(MeshAttribute::Vertex::Position);
 	TEdgeAttributesRef<bool> EdgeHardnesses = OutRawMesh.EdgeAttributes().GetAttributesRef<bool>(MeshAttribute::Edge::IsHard);
-	TEdgeAttributesRef<float> EdgeCreaseSharpnesses = OutRawMesh.EdgeAttributes().GetAttributesRef<float>(MeshAttribute::Edge::CreaseSharpness);
 	TPolygonGroupAttributesRef<FName> PolygonGroupImportedMaterialSlotNames = OutRawMesh.PolygonGroupAttributes().GetAttributesRef<FName>(MeshAttribute::PolygonGroup::ImportedMaterialSlotName);
 	TVertexInstanceAttributesRef<FVector> VertexInstanceNormals = OutRawMesh.VertexInstanceAttributes().GetAttributesRef<FVector>(MeshAttribute::VertexInstance::Normal);
 	TVertexInstanceAttributesRef<FVector> VertexInstanceTangents = OutRawMesh.VertexInstanceAttributes().GetAttributesRef<FVector>(MeshAttribute::VertexInstance::Tangent);
@@ -1045,16 +1042,16 @@ void FMeshMergeHelpers::CalculateTextureCoordinateBoundsForRawMesh(const FMeshDe
 {
 	TVertexInstanceAttributesConstRef<FVector2D> VertexInstanceUVs = InRawMesh.VertexInstanceAttributes().GetAttributesRef<FVector2D>(MeshAttribute::VertexInstance::TextureCoordinate);
 	OutBounds.Empty();
-	for (const FPolygonID& PolygonID : InRawMesh.Polygons().GetElementIDs())
+	for (const FTriangleID TriangleID : InRawMesh.Polygons().GetElementIDs())
 	{
-		int32 MaterialIndex = InRawMesh.GetPolygonPolygonGroup(PolygonID).GetValue();
+		int32 MaterialIndex = InRawMesh.GetTrianglePolygonGroup(TriangleID).GetValue();
 		if (OutBounds.Num() <= MaterialIndex)
 			OutBounds.SetNumZeroed(MaterialIndex + 1);
 		{
-			TArray<FVertexInstanceID> PolygonVertexInstances = InRawMesh.GetPolygonVertexInstances(PolygonID);
-			for (const FVertexInstanceID& VertexInstanceID : PolygonVertexInstances)
+			TArrayView<const FVertexInstanceID> TriangleVertexInstances = InRawMesh.GetTriangleVertexInstances(TriangleID);
+			for (const FVertexInstanceID VertexInstanceID : TriangleVertexInstances)
 			{
-				for (int32 UVIndex = 0; UVIndex < VertexInstanceUVs.GetNumIndices(); ++UVIndex)
+				for (int32 UVIndex = 0; UVIndex < VertexInstanceUVs.GetNumChannels(); ++UVIndex)
 				{
 					OutBounds[MaterialIndex] += VertexInstanceUVs.Get(VertexInstanceID, UVIndex);
 				}
@@ -1090,7 +1087,7 @@ bool FMeshMergeHelpers::PropagatePaintedColorsToRawMesh(const UStaticMeshCompone
 				int32 CurrentWedgeIndex = 0;
 				for (const FPolygonID PolygonID : RawMesh.Polygons().GetElementIDs())
 				{
-					const TArray<FTriangleID>& TriangleIDs = RawMesh.GetPolygonTriangleIDs(PolygonID);
+					TArrayView<const FTriangleID> TriangleIDs = RawMesh.GetPolygonTriangleIDs(PolygonID);
 					for (const FTriangleID TriangleID : TriangleIDs)
 					{
 						for (int32 Corner = 0; Corner < 3; ++Corner, ++CurrentWedgeIndex)
@@ -1138,7 +1135,7 @@ bool FMeshMergeHelpers::PropagatePaintedColorsToRawMesh(const UStaticMeshCompone
 						{
 							WedgeColor = ColorVertexBuffer.VertexColor(VertIndex);
 						}
-						const TArray<FVertexInstanceID>& VertexInstances = RawMesh.GetVertexVertexInstances(VertexID);
+						TArrayView<const FVertexInstanceID> VertexInstances = RawMesh.GetVertexVertexInstances(VertexID);
 						for (const FVertexInstanceID& VertexInstanceID : VertexInstances)
 						{
 							VertexInstanceColors[VertexInstanceID] = FLinearColor(WedgeColor);
@@ -1185,25 +1182,26 @@ void FMeshMergeHelpers::AppendRawMesh(FMeshDescription& InTarget, const FMeshDes
 {
 	TRACE_CPUPROFILER_EVENT_SCOPE(FMeshMergeHelpers::AppendRawMesh)
 
-	TVertexAttributesConstRef<FVector> SourceVertexPositions = InSource.VertexAttributes().GetAttributesRef<FVector>(MeshAttribute::Vertex::Position);
-	TEdgeAttributesConstRef<bool> SourceEdgeHardnesses = InSource.EdgeAttributes().GetAttributesRef<bool>(MeshAttribute::Edge::IsHard);
-	TEdgeAttributesConstRef<float> SourceEdgeCreaseSharpnesses = InSource.EdgeAttributes().GetAttributesRef<float>(MeshAttribute::Edge::CreaseSharpness);
-	TPolygonGroupAttributesConstRef<FName> SourcePolygonGroupImportedMaterialSlotNames = InSource.PolygonGroupAttributes().GetAttributesRef<FName>(MeshAttribute::PolygonGroup::ImportedMaterialSlotName);
-	TVertexInstanceAttributesConstRef<FVector> SourceVertexInstanceNormals = InSource.VertexInstanceAttributes().GetAttributesRef<FVector>(MeshAttribute::VertexInstance::Normal);
-	TVertexInstanceAttributesConstRef<FVector> SourceVertexInstanceTangents = InSource.VertexInstanceAttributes().GetAttributesRef<FVector>(MeshAttribute::VertexInstance::Tangent);
-	TVertexInstanceAttributesConstRef<float> SourceVertexInstanceBinormalSigns = InSource.VertexInstanceAttributes().GetAttributesRef<float>(MeshAttribute::VertexInstance::BinormalSign);
-	TVertexInstanceAttributesConstRef<FVector4> SourceVertexInstanceColors = InSource.VertexInstanceAttributes().GetAttributesRef<FVector4>(MeshAttribute::VertexInstance::Color);
-	TVertexInstanceAttributesConstRef<FVector2D> SourceVertexInstanceUVs = InSource.VertexInstanceAttributes().GetAttributesRef<FVector2D>(MeshAttribute::VertexInstance::TextureCoordinate);
-
-	TVertexAttributesRef<FVector> TargetVertexPositions = InTarget.VertexAttributes().GetAttributesRef<FVector>(MeshAttribute::Vertex::Position);
-	TEdgeAttributesRef<bool> TargetEdgeHardnesses = InTarget.EdgeAttributes().GetAttributesRef<bool>(MeshAttribute::Edge::IsHard);
-	TEdgeAttributesRef<float> TargetEdgeCreaseSharpnesses = InTarget.EdgeAttributes().GetAttributesRef<float>(MeshAttribute::Edge::CreaseSharpness);
-	TPolygonGroupAttributesRef<FName> TargetPolygonGroupImportedMaterialSlotNames = InTarget.PolygonGroupAttributes().GetAttributesRef<FName>(MeshAttribute::PolygonGroup::ImportedMaterialSlotName);
-	TVertexInstanceAttributesRef<FVector> TargetVertexInstanceNormals = InTarget.VertexInstanceAttributes().GetAttributesRef<FVector>(MeshAttribute::VertexInstance::Normal);
-	TVertexInstanceAttributesRef<FVector> TargetVertexInstanceTangents = InTarget.VertexInstanceAttributes().GetAttributesRef<FVector>(MeshAttribute::VertexInstance::Tangent);
-	TVertexInstanceAttributesRef<float> TargetVertexInstanceBinormalSigns = InTarget.VertexInstanceAttributes().GetAttributesRef<float>(MeshAttribute::VertexInstance::BinormalSign);
-	TVertexInstanceAttributesRef<FVector4> TargetVertexInstanceColors = InTarget.VertexInstanceAttributes().GetAttributesRef<FVector4>(MeshAttribute::VertexInstance::Color);
-	TVertexInstanceAttributesRef<FVector2D> TargetVertexInstanceUVs = InTarget.VertexInstanceAttributes().GetAttributesRef<FVector2D>(MeshAttribute::VertexInstance::TextureCoordinate);
+	FStaticMeshConstAttributes SourceAttributes(InSource);
+	TVertexAttributesConstRef<FVector> SourceVertexPositions = SourceAttributes.GetVertexPositions();
+	TEdgeAttributesConstRef<bool> SourceEdgeHardnesses = SourceAttributes.GetEdgeHardnesses();
+	TPolygonGroupAttributesConstRef<FName> SourcePolygonGroupImportedMaterialSlotNames = SourceAttributes.GetPolygonGroupMaterialSlotNames();
+	TVertexInstanceAttributesConstRef<FVector> SourceVertexInstanceNormals = SourceAttributes.GetVertexInstanceNormals();
+	TVertexInstanceAttributesConstRef<FVector> SourceVertexInstanceTangents = SourceAttributes.GetVertexInstanceTangents();
+	TVertexInstanceAttributesConstRef<float> SourceVertexInstanceBinormalSigns = SourceAttributes.GetVertexInstanceBinormalSigns();
+	TVertexInstanceAttributesConstRef<FVector4> SourceVertexInstanceColors = SourceAttributes.GetVertexInstanceColors();
+	TVertexInstanceAttributesConstRef<FVector2D> SourceVertexInstanceUVs = SourceAttributes.GetVertexInstanceUVs();
+
+
+	FStaticMeshAttributes TargetAttributes(InTarget);
+	TVertexAttributesRef<FVector> TargetVertexPositions = TargetAttributes.GetVertexPositions();
+	TEdgeAttributesRef<bool> TargetEdgeHardnesses = TargetAttributes.GetEdgeHardnesses();
+	TPolygonGroupAttributesRef<FName> TargetPolygonGroupImportedMaterialSlotNames = TargetAttributes.GetPolygonGroupMaterialSlotNames();
+	TVertexInstanceAttributesRef<FVector> TargetVertexInstanceNormals = TargetAttributes.GetVertexInstanceNormals();
+	TVertexInstanceAttributesRef<FVector> TargetVertexInstanceTangents = TargetAttributes.GetVertexInstanceTangents();
+	TVertexInstanceAttributesRef<float> TargetVertexInstanceBinormalSigns = TargetAttributes.GetVertexInstanceBinormalSigns();
+	TVertexInstanceAttributesRef<FVector4> TargetVertexInstanceColors = TargetAttributes.GetVertexInstanceColors();
+	TVertexInstanceAttributesRef<FVector2D> TargetVertexInstanceUVs = TargetAttributes.GetVertexInstanceUVs();
 
 	InTarget.ReserveNewVertices(InSource.Vertices().Num());
 	InTarget.ReserveNewVertexInstances(InSource.VertexInstances().Num());
@@ -1249,9 +1247,9 @@ void FMeshMergeHelpers::AppendRawMesh(FMeshDescription& InTarget, const FMeshDes
 	}
 
 	//Append VertexInstances
-	if (SourceVertexInstanceUVs.GetNumIndices() > TargetVertexInstanceUVs.GetNumIndices())
+	if (SourceVertexInstanceUVs.GetNumChannels() > TargetVertexInstanceUVs.GetNumChannels())
 	{
-		TargetVertexInstanceUVs.SetNumIndices(SourceVertexInstanceUVs.GetNumIndices());
+		TargetVertexInstanceUVs.SetNumChannels(SourceVertexInstanceUVs.GetNumChannels());
 	}
 	TMap<FVertexInstanceID, FVertexInstanceID> SourceToTargetVertexInstanceID;
 	SourceToTargetVertexInstanceID.Reserve(InSource.VertexInstances().Num());
@@ -1263,9 +1261,9 @@ void FMeshMergeHelpers::AppendRawMesh(FMeshDescription& InTarget, const FMeshDes
 		TargetVertexInstanceBinormalSigns[TargetVertexInstanceID] = SourceVertexInstanceBinormalSigns[SourceVertexInstanceID];
 		TargetVertexInstanceNormals[TargetVertexInstanceID] = SourceVertexInstanceNormals[SourceVertexInstanceID];
 		TargetVertexInstanceColors[TargetVertexInstanceID] = SourceVertexInstanceColors[SourceVertexInstanceID];
-		for (int32 UVIndex = 0; UVIndex < TargetVertexInstanceUVs.GetNumIndices(); ++UVIndex)
+		for (int32 UVIndex = 0; UVIndex < TargetVertexInstanceUVs.GetNumChannels(); ++UVIndex)
 		{
-			FVector2D SourceUV = SourceVertexInstanceUVs.GetNumIndices() > UVIndex ? SourceVertexInstanceUVs.Get(SourceVertexInstanceID, UVIndex) : FVector2D(0.0f, 0.0f);
+			FVector2D SourceUV = SourceVertexInstanceUVs.GetNumChannels() > UVIndex ? SourceVertexInstanceUVs.Get(SourceVertexInstanceID, UVIndex) : FVector2D(0.0f, 0.0f);
 			TargetVertexInstanceUVs.Set(TargetVertexInstanceID, UVIndex, SourceUV);
 		}
 		SourceToTargetVertexInstanceID.Add(SourceVertexInstanceID, TargetVertexInstanceID);
@@ -1280,7 +1278,6 @@ void FMeshMergeHelpers::AppendRawMesh(FMeshDescription& InTarget, const FMeshDes
 		const FVertexID EdgeVertex1 = InSource.GetEdgeVertex(SourceEdgeID, 1);
 		const FEdgeID TargetEdgeID = InTarget.CreateEdge(SourceToTargetVertexID[EdgeVertex0], SourceToTargetVertexID[EdgeVertex1]);
 		TargetEdgeHardnesses[TargetEdgeID] = SourceEdgeHardnesses[SourceEdgeID];
-		TargetEdgeCreaseSharpnesses[TargetEdgeID] = SourceEdgeCreaseSharpnesses[SourceEdgeID];
 		SourceToTargetEdgeID.Add(SourceEdgeID, TargetEdgeID);
 	}
 
@@ -1294,10 +1291,11 @@ void FMeshMergeHelpers::AppendRawMesh(FMeshDescription& InTarget, const FMeshDes
 		}
 	};
 
+	// @todo: make this work in terms of triangles instead of polygons
 	//Append Polygons
 	for (const FPolygonID SourcePolygonID : InSource.Polygons().GetElementIDs())
 	{
-		const TArray<FVertexInstanceID>& SourceVertexInstanceIDs = InSource.GetPolygonVertexInstances(SourcePolygonID);
+		TArray<FVertexInstanceID> SourceVertexInstanceIDs = InSource.GetPolygonVertexInstances(SourcePolygonID);
 		const FPolygonGroupID PolygonGroupID = InSource.GetPolygonPolygonGroup(SourcePolygonID);
 
 		TArray<FVertexInstanceID> ContourVertexInstances;
@@ -1346,7 +1344,7 @@ void FMeshMergeHelpers::MergeImpostersToRawMesh(TArray<const UStaticMeshComponen
 			const int32 UVOneIndex = 2;
 			const int32 UVTwoIndex = UVOneIndex + 1;
 			TVertexInstanceAttributesRef<FVector2D> VertexInstanceUVs = ImposterMeshAttributes.GetVertexInstanceUVs();
-			VertexInstanceUVs.SetNumIndices(UVTwoIndex + 1);
+			VertexInstanceUVs.SetNumChannels(UVTwoIndex + 1);
 			const int32 NumIndices = ImposterMesh.VertexInstances().Num();
 			const FTransform& ActorToWorld = Component->GetOwner()->GetActorTransform();
 			const FVector ActorPosition = ActorToWorld.TransformPosition(FVector::ZeroVector) - InPivot;
@@ -1386,8 +1384,8 @@ void FMeshMergeHelpers::MergeImpostersToRawMesh(TArray<const UStaticMeshComponen
 		{
 			UMaterialInterface* MaterialUseBySection = OutImposterMaterials[SectionImposterUniqueMaterialIndex[SectionIndex++]];
 			FPolygonGroupID* ExistTargetPolygonGroupID = ImposterMaterialToPolygonGroupID.Find(MaterialUseBySection);
-			FPolygonGroupID MatchTargetPolygonGroupID = ExistTargetPolygonGroupID == nullptr ? FPolygonGroupID::Invalid : *ExistTargetPolygonGroupID;
-			if (MatchTargetPolygonGroupID == FPolygonGroupID::Invalid)
+			FPolygonGroupID MatchTargetPolygonGroupID = ExistTargetPolygonGroupID == nullptr ? INDEX_NONE : *ExistTargetPolygonGroupID;
+			if (MatchTargetPolygonGroupID == INDEX_NONE)
 			{
 				MatchTargetPolygonGroupID = InRawMesh.CreatePolygonGroup();
 				//use the material name to fill the imported material name. Material name will be unique
diff --git a/Engine/Source/Developer/MeshMergeUtilities/Private/MeshMergeUtilities.cpp b/Engine/Source/Developer/MeshMergeUtilities/Private/MeshMergeUtilities.cpp
--- a/Engine/Source/Developer/MeshMergeUtilities/Private/MeshMergeUtilities.cpp
+++ b/Engine/Source/Developer/MeshMergeUtilities/Private/MeshMergeUtilities.cpp
@@ -231,7 +231,7 @@ void FMeshMergeUtilities::BakeMaterialsForComponent(TArray<TWeakObjectPtr<UObjec
 					const int32 LightMapUVIndex = Adapter->LightmapUVIndex();
 					
 					TVertexInstanceAttributesConstRef<FVector2D> VertexInstanceUVs = MeshSettings.RawMeshDescription->VertexInstanceAttributes().GetAttributesRef<FVector2D>(MeshAttribute::VertexInstance::TextureCoordinate);
-					if (bNeedsUniqueUVs && MeshSettings.TextureCoordinateIndex != LightMapUVIndex && VertexInstanceUVs.GetNumElements() > 0 && VertexInstanceUVs.GetNumIndices() > LightMapUVIndex)
+					if (bNeedsUniqueUVs && MeshSettings.TextureCoordinateIndex != LightMapUVIndex && VertexInstanceUVs.GetNumElements() > 0 && VertexInstanceUVs.GetNumChannels() > LightMapUVIndex)
 					{
 						MeshSettings.TextureCoordinateIndex = LightMapUVIndex;
 					}
@@ -1399,7 +1399,7 @@ void FMeshMergeUtilities::CreateProxyMesh(const TArray<UStaticMeshComponent*>& I
 						if (InMeshProxySettings.bReuseMeshLightmapUVs &&
 							LightMapCoordinateIndex > 0 &&
 							VertexInstanceUVs.GetNumElements() > 0 &&
-							VertexInstanceUVs.GetNumIndices() > LightMapCoordinateIndex)
+							VertexInstanceUVs.GetNumChannels() > LightMapCoordinateIndex)
 						{
 							MeshSettings.CustomTextureCoordinates.Reset(VertexInstanceUVs.GetNumElements());
 							for (const FVertexInstanceID VertexInstanceID : MeshDescription->VertexInstances().GetElementIDs())
@@ -2161,7 +2161,7 @@ void FMeshMergeUtilities::MergeComponentsToStaticMesh(const TArray<UPrimitiveCom
 						const int32 LightMapUVIndex = StaticMeshComponentsToMerge[Key.GetMeshIndex()]->GetStaticMesh()->LightMapCoordinateIndex;
 
 						TVertexInstanceAttributesRef<FVector2D> VertexInstanceUVs = MeshData.RawMeshDescription->VertexInstanceAttributes().GetAttributesRef<FVector2D>(MeshAttribute::VertexInstance::TextureCoordinate);
-						if (InSettings.bReuseMeshLightmapUVs && VertexInstanceUVs.GetNumElements() > 0 && VertexInstanceUVs.GetNumIndices() > LightMapUVIndex)
+						if (InSettings.bReuseMeshLightmapUVs && VertexInstanceUVs.GetNumElements() > 0 && VertexInstanceUVs.GetNumChannels() > LightMapUVIndex)
 						{
 							MeshData.TextureCoordinateIndex = LightMapUVIndex;
 						}
@@ -2424,7 +2424,7 @@ void FMeshMergeUtilities::MergeComponentsToStaticMesh(const TArray<UPrimitiveCom
 				if(RawMesh.Vertices().Num())
 				{
 					const TVertexInstanceAttributesRef<const FVector2D> VertexInstanceUVs = RawMesh.VertexInstanceAttributes().GetAttributesRef<FVector2D>(MeshAttribute::VertexInstance::TextureCoordinate);
-					MergedMatUVChannel = FMath::Max(MergedMatUVChannel, VertexInstanceUVs.GetNumIndices());
+					MergedMatUVChannel = FMath::Max(MergedMatUVChannel, VertexInstanceUVs.GetNumChannels());
 				}
 			}
 		}
@@ -2448,7 +2448,7 @@ void FMeshMergeUtilities::MergeComponentsToStaticMesh(const TArray<UPrimitiveCom
 					if (RawMesh->Vertices().Num())
 					{
 						TVertexInstanceAttributesRef<FVector2D> VertexInstanceUVs = RawMesh->VertexInstanceAttributes().GetAttributesRef<FVector2D>(MeshAttribute::VertexInstance::TextureCoordinate);
-						int32 NumUVChannel = FMath::Min(VertexInstanceUVs.GetNumIndices(), (int32)MAX_MESH_TEXTURE_COORDS);
+						int32 NumUVChannel = FMath::Min(VertexInstanceUVs.GetNumChannels(), (int32)MAX_MESH_TEXTURE_COORDS);
 						for (int32 UVChannelIdx = 0; UVChannelIdx < NumUVChannel; ++UVChannelIdx)
 						{
 							int32 VertexIndex = 0;
@@ -2487,7 +2487,7 @@ void FMeshMergeUtilities::MergeComponentsToStaticMesh(const TArray<UPrimitiveCom
 
 						if (InSettings.bCreateMergedMaterial && MeshData.CustomTextureCoordinates.Num() > 0)
 						{
-							VertexInstanceUVs.SetNumIndices(MergedMatUVChannel + 1);
+							VertexInstanceUVs.SetNumChannels(MergedMatUVChannel + 1);
 
 							int32 VertexIndex = 0;
 							for(FVertexInstanceID VertexInstanceID : RawMesh->VertexInstances().GetElementIDs())
diff --git a/Engine/Source/Developer/MeshSimplifier/Private/QuadricMeshReduction.cpp b/Engine/Source/Developer/MeshSimplifier/Private/QuadricMeshReduction.cpp
--- a/Engine/Source/Developer/MeshSimplifier/Private/QuadricMeshReduction.cpp
+++ b/Engine/Source/Developer/MeshSimplifier/Private/QuadricMeshReduction.cpp
@@ -227,7 +227,7 @@ public:
 		{
 			const FPolygonGroupID PolygonGroupID = InMesh.GetPolygonPolygonGroup(PolygonID);
 
-			const TArray<FTriangleID>& TriangleIDs = InMesh.GetPolygonTriangleIDs(PolygonID);
+			TArrayView<const FTriangleID> TriangleIDs = InMesh.GetPolygonTriangleIDs(PolygonID);
 			for (int32 TriangleIndex = 0; TriangleIndex < TriangleIDs.Num(); ++TriangleIndex)
 			{
 				const FTriangleID TriangleID = TriangleIDs[TriangleIndex];
@@ -283,7 +283,7 @@ public:
 
 					for (int32 UVIndex = 0; UVIndex < NumTexCoords; UVIndex++)
 					{
-						if (UVIndex < InVertexUVs.GetNumIndices())
+						if (UVIndex < InVertexUVs.GetNumChannels())
 						{
 							NewVert.TexCoords[UVIndex] = InVertexUVs.Get(VertexInstanceID, UVIndex);
 							InMeshNumTexCoords = FMath::Max(UVIndex + 1, InMeshNumTexCoords);
@@ -416,7 +416,7 @@ public:
 
 			for (int32 TexCoordIndex = 0; TexCoordIndex < NumTexCoords; TexCoordIndex++)
 			{
-				if (TexCoordIndex >= InVertexUVs.GetNumIndices())
+				if (TexCoordIndex >= InVertexUVs.GetNumChannels())
 				{
 					TexCoordWeights[2 * TexCoordIndex + 0] = 0.0f;
 					TexCoordWeights[2 * TexCoordIndex + 1] = 0.0f;
@@ -448,11 +448,7 @@ public:
 
 		{
 			//Empty the destination mesh
-			OutReducedMesh.PolygonGroups().Reset();
-			OutReducedMesh.Polygons().Reset();
-			OutReducedMesh.Edges().Reset();
-			OutReducedMesh.VertexInstances().Reset();
-			OutReducedMesh.Vertices().Reset();
+			OutReducedMesh.Empty();
 
 			//Fill the PolygonGroups from the InMesh
 			for (const FPolygonGroupID& PolygonGroupID : InMesh.PolygonGroups().GetElementIDs())
@@ -480,13 +476,12 @@ public:
 			TVertexInstanceAttributesRef<FVector2D> OutVertexUVs = OutReducedMesh.VertexInstanceAttributes().GetAttributesRef<FVector2D>(MeshAttribute::VertexInstance::TextureCoordinate);
 
 			//Specify the number of texture coords in this mesh description
-			OutVertexUVs.SetNumIndices(InMeshNumTexCoords);
+			OutVertexUVs.SetNumChannels(InMeshNumTexCoords);
 
 			//Vertex instances and Polygons
 			for (int32 TriangleIndex = 0; TriangleIndex < (int32)NumTris; TriangleIndex++)
 			{
-				TArray<FVertexInstanceID> CornerInstanceIDs;
-				CornerInstanceIDs.SetNum(3);
+				FVertexInstanceID CornerInstanceIDs[3];
 
 				FVertexID CornerVerticesIDs[3];
 				for (int32 CornerIndex = 0; CornerIndex < 3; CornerIndex++)
@@ -520,7 +515,7 @@ public:
 				
 				// material index
 				int32 MaterialIndex = Verts[Indexes[3 * TriangleIndex]].MaterialIndex;
-				FPolygonGroupID MaterialPolygonGroupID = FPolygonGroupID::Invalid;
+				FPolygonGroupID MaterialPolygonGroupID = INDEX_NONE;
 				if (!PolygonGroupMapping.Contains(MaterialIndex))
 				{
 					FPolygonGroupID PolygonGroupID(MaterialIndex);
@@ -531,7 +526,7 @@ public:
 					InMesh.PolygonGroupAttributes().ForEach(
 						[&OutReducedMesh, PolygonGroupID, MaterialPolygonGroupID](const FName Name, const auto ArrayRef)
 						{
-							for (int32 Index = 0; Index < ArrayRef.GetNumIndices(); ++Index)
+							for (int32 Index = 0; Index < ArrayRef.GetNumChannels(); ++Index)
 							{
 								// Only copy shared attribute values, since input mesh description can differ from output mesh description
 								const auto& Value = ArrayRef.Get(PolygonGroupID, Index);
@@ -551,7 +546,7 @@ public:
 
 				// Insert a polygon into the mesh
 				TArray<FEdgeID> NewEdgeIDs;
-				const FPolygonID NewPolygonID = OutReducedMesh.CreatePolygon(MaterialPolygonGroupID, CornerInstanceIDs, &NewEdgeIDs);
+				const FTriangleID NewTriangleID = OutReducedMesh.CreateTriangle(MaterialPolygonGroupID, CornerInstanceIDs, &NewEdgeIDs);
 				for (const FEdgeID NewEdgeID : NewEdgeIDs)
 				{
 					// @todo: set NewEdgeID edge hardness?
diff --git a/Engine/Source/Editor/UnrealEd/Private/Fbx/FbxMainExport.cpp b/Engine/Source/Editor/UnrealEd/Private/Fbx/FbxMainExport.cpp
--- a/Engine/Source/Editor/UnrealEd/Private/Fbx/FbxMainExport.cpp
+++ b/Engine/Source/Editor/UnrealEd/Private/Fbx/FbxMainExport.cpp
@@ -1056,7 +1056,6 @@ void FFbxExporter::ExportBSP( UModel* Model, bool bSelectedOnly )
 			TVertexInstanceAttributesRef<FVector4> Colors = MeshAttributes.GetVertexInstanceColors();
 			TVertexInstanceAttributesRef<FVector> Normals = MeshAttributes.GetVertexInstanceNormals();
 			TEdgeAttributesRef<bool> EdgeHardnesses = MeshAttributes.GetEdgeHardnesses();
-			TEdgeAttributesRef<float> EdgeCreaseSharpnesses = MeshAttributes.GetEdgeCreaseSharpnesses();
 
 			
 			UMaterialInterface*	Material = Poly.Material;
@@ -1126,7 +1125,6 @@ void FFbxExporter::ExportBSP( UModel* Model, bool bSelectedOnly )
 				for (const FEdgeID EdgeID : NewEdgeIDs)
 				{
 					EdgeHardnesses[EdgeID] = false;
-					EdgeCreaseSharpnesses[EdgeID] = 0.0f;
 				}
 
 				//Add to the smoothGroup array so we can compute hard edge later
diff --git a/Engine/Source/Editor/UnrealEd/Private/Fbx/FbxStaticMeshImport.cpp b/Engine/Source/Editor/UnrealEd/Private/Fbx/FbxStaticMeshImport.cpp
--- a/Engine/Source/Editor/UnrealEd/Private/Fbx/FbxStaticMeshImport.cpp
+++ b/Engine/Source/Editor/UnrealEd/Private/Fbx/FbxStaticMeshImport.cpp
@@ -600,6 +600,12 @@ bool UnFbx::FFbxImporter::BuildStaticMeshFromGeometry(FbxNode* Node, UStaticMesh
 	//
 	bool bHasNonDegeneratePolygons = false;
 	{
+		MeshDescription->SuspendVertexInstanceIndexing();
+		MeshDescription->SuspendEdgeIndexing();
+		MeshDescription->SuspendPolygonIndexing();
+		MeshDescription->SuspendPolygonGroupIndexing();
+		MeshDescription->SuspendUVIndexing();
+
 		TRACE_CPUPROFILER_EVENT_SCOPE(BuildTriangles);
 
 		// Construct the matrices for the conversion from right handed to left handed system
@@ -626,27 +632,41 @@ bool UnFbx::FFbxImporter::BuildStaticMeshFromGeometry(FbxNode* Node, UStaticMesh
 		TVertexInstanceAttributesRef<FVector4> VertexInstanceColors = Attributes.GetVertexInstanceColors();
 		TVertexInstanceAttributesRef<FVector2D> VertexInstanceUVs = Attributes.GetVertexInstanceUVs();
 		TEdgeAttributesRef<bool> EdgeHardnesses = Attributes.GetEdgeHardnesses();
-		TEdgeAttributesRef<float> EdgeCreaseSharpnesses = Attributes.GetEdgeCreaseSharpnesses();
 		TPolygonGroupAttributesRef<FName> PolygonGroupImportedMaterialSlotNames = Attributes.GetPolygonGroupMaterialSlotNames();
 
 		int32 VertexOffset = MeshDescription->Vertices().Num();
 		int32 VertexInstanceOffset = MeshDescription->VertexInstances().Num();
 		int32 PolygonOffset = MeshDescription->Polygons().Num();
 
+		// The below code expects Num() to be equivalent to GetArraySize(), i.e. that all added elements are appended, not inserted into existing gaps
+		check(VertexOffset == MeshDescription->Vertices().GetArraySize());
+		check(VertexInstanceOffset == MeshDescription->VertexInstances().GetArraySize());
+		check(PolygonOffset == MeshDescription->Polygons().GetArraySize());
+
 		TMap<int32, FPolygonGroupID> PolygonGroupMapping;
 
 		// When importing multiple mesh pieces to the same static mesh.  Ensure each mesh piece has the same number of Uv's
-		int32 ExistingUVCount = VertexInstanceUVs.GetNumIndices();
+		int32 ExistingUVCount = VertexInstanceUVs.GetNumChannels();
 
 		int32 NumUVs = FMath::Max(FBXUVs.UniqueUVCount, ExistingUVCount);
 		NumUVs = FMath::Min<int32>(MAX_MESH_TEXTURE_COORDS_MD, NumUVs);
 		// At least one UV set must exist.  
+		// @todo: this needn't be true; we should be able to handle zero UV channels
 		NumUVs = FMath::Max(1, NumUVs);
 
 		//Make sure all Vertex instance have the correct number of UVs
-		VertexInstanceUVs.SetNumIndices( NumUVs );
+		VertexInstanceUVs.SetNumChannels( NumUVs );
+		MeshDescription->SetNumUVChannels(NumUVs);
+
+		TArray<int32> UVOffsets;
+		UVOffsets.SetNumUninitialized(NumUVs);
+		for (int UVChannel = 0; UVChannel < NumUVs; UVChannel++)
+		{
+			UVOffsets[UVChannel] = MeshDescription->UVs(UVChannel).GetArraySize();
+		}
 
 		//Fill the vertex array
+		MeshDescription->ReserveNewVertices(VertexCount);
 		for (int32 VertexIndex = 0; VertexIndex < VertexCount; ++VertexIndex)
 		{
 			int32 RealVertexIndex = VertexOffset + VertexIndex;
@@ -663,6 +683,24 @@ bool UnFbx::FFbxImporter::BuildStaticMeshFromGeometry(FbxNode* Node, UStaticMesh
 			}
 		}
 		
+		// Fill the UV arrays
+		for (int32 UVLayerIndex = 0; UVLayerIndex < FBXUVs.UniqueUVCount; UVLayerIndex++)
+		{
+			check(FBXUVs.LayerElementUV[UVLayerIndex]);
+			if (FBXUVs.LayerElementUV[UVLayerIndex] != nullptr)
+			{
+				int32 UVCount = FBXUVs.LayerElementUV[UVLayerIndex]->GetDirectArray().GetCount();
+				TUVAttributesRef<FVector2D> UVCoordinates = MeshDescription->UVAttributes(UVLayerIndex).GetAttributesRef<FVector2D>(MeshAttribute::UV::UVCoordinate);
+				MeshDescription->ReserveNewUVs(UVCount, UVLayerIndex);
+				for (int32 UVIndex = 0; UVIndex < UVCount; UVIndex++)
+				{
+					FUVID UVID = MeshDescription->CreateUV(UVLayerIndex);
+					FbxVector2 UVVector = FBXUVs.LayerElementUV[UVLayerIndex]->GetDirectArray().GetAt(UVIndex);
+					UVCoordinates[UVIndex] = FVector2D(static_cast<float>(UVVector[0]), 1.0f - static_cast<float>(UVVector[1]));	// flip the Y of UVs for DirectX
+				}
+			}
+		}
+
 		TMap<uint64, int32> RemapEdgeID;
 		{
 			TRACE_CPUPROFILER_EVENT_SCOPE(BuildMeshEdgeVertices);
@@ -746,10 +784,10 @@ bool UnFbx::FFbxImporter::BuildStaticMeshFromGeometry(FbxNode* Node, UStaticMesh
 				CornerInstanceIDs.AddUninitialized(PolygonVertexCount);
 				CornerVerticesIDs.Reset();
 				CornerVerticesIDs.AddUninitialized(PolygonVertexCount);
-				for (int32 CornerIndex = 0; CornerIndex < PolygonVertexCount; CornerIndex++, CurrentVertexInstanceIndex++)
+				for (int32 CornerIndex = 0; CornerIndex < PolygonVertexCount; CornerIndex++)
 				{
-					int32 VertexInstanceIndex = VertexInstanceOffset + CurrentVertexInstanceIndex;
-					int32 RealFbxVertexIndex = SkippedVertexInstance + CurrentVertexInstanceIndex;
+					int32 VertexInstanceIndex = VertexInstanceOffset + CurrentVertexInstanceIndex + CornerIndex;
+					int32 RealFbxVertexIndex = SkippedVertexInstance + CurrentVertexInstanceIndex + CornerIndex;
 					const FVertexInstanceID VertexInstanceID(VertexInstanceIndex);
 					CornerInstanceIDs[CornerIndex] = VertexInstanceID;
 					const int32 ControlPointIndex = Mesh->GetPolygonVertex(PolygonIndex, CornerIndex);
@@ -917,7 +955,7 @@ bool UnFbx::FFbxImporter::BuildStaticMeshFromGeometry(FbxNode* Node, UStaticMesh
 				{
 					UMaterialInterface* Material = MeshMaterials.IsValidIndex(RealMaterialIndex) ? MeshMaterials[RealMaterialIndex].Material : UMaterial::GetDefaultMaterial(MD_Surface);
 					FName ImportedMaterialSlotName = MeshMaterials.IsValidIndex(RealMaterialIndex) ? FName(*MeshMaterials[RealMaterialIndex].GetName()) : (Material != nullptr ? FName(*Material->GetName()) : NAME_None);
-					FPolygonGroupID ExistingPolygonGroup = FPolygonGroupID::Invalid;
+					FPolygonGroupID ExistingPolygonGroup = INDEX_NONE;
 					for (const FPolygonGroupID PolygonGroupID : MeshDescription->PolygonGroups().GetElementIDs())
 					{
 						if (PolygonGroupImportedMaterialSlotNames[PolygonGroupID] == ImportedMaterialSlotName)
@@ -926,7 +964,7 @@ bool UnFbx::FFbxImporter::BuildStaticMeshFromGeometry(FbxNode* Node, UStaticMesh
 							break;
 						}
 					}
-					if (ExistingPolygonGroup == FPolygonGroupID::Invalid)
+					if (ExistingPolygonGroup == INDEX_NONE)
 					{
 						ExistingPolygonGroup = MeshDescription->CreatePolygonGroup();
 						PolygonGroupImportedMaterialSlotNames[ExistingPolygonGroup] = ImportedMaterialSlotName;
@@ -949,7 +987,7 @@ bool UnFbx::FFbxImporter::BuildStaticMeshFromGeometry(FbxNode* Node, UStaticMesh
 						EdgeVertexIDs[1] = CornerVerticesIDs[CornerIndices[1]];
 
 						FEdgeID MatchEdgeId = MeshDescription->GetVertexPairEdge(EdgeVertexIDs[0], EdgeVertexIDs[1]);
-						if (MatchEdgeId == FEdgeID::Invalid)
+						if (MatchEdgeId == INDEX_NONE)
 						{
 							MatchEdgeId = MeshDescription->CreateEdge(EdgeVertexIDs[0], EdgeVertexIDs[1]);
 						}
@@ -981,7 +1019,6 @@ bool UnFbx::FFbxImporter::BuildStaticMeshFromGeometry(FbxNode* Node, UStaticMesh
 							EdgeIndex = Mesh->GetMeshEdgeIndexForPolygon(PolygonIndex, PolygonEdgeNumber);
 						}
 					
-						EdgeCreaseSharpnesses[MatchEdgeId] = (float)Mesh->GetEdgeCreaseInfo(EdgeIndex);
 						if (!EdgeHardnesses[MatchEdgeId])
 						{
 							if (bSmoothingAvailable && SmoothingInfo)
@@ -1007,16 +1044,54 @@ bool UnFbx::FFbxImporter::BuildStaticMeshFromGeometry(FbxNode* Node, UStaticMesh
 				}
 
 				FPolygonGroupID PolygonGroupID = PolygonGroupMapping[RealMaterialIndex];
-				// Insert a polygon into the mesh
+				// Insert a triangle into the mesh
+				// @note: we only ever import triangulated meshes currently. This could easily change as we have the infrastructure set up for arbitrary ngons.
+				// However, I think either we need to triangulate in the exact same way as Maya/Max if we do that.
+				check(CornerInstanceIDs.Num() == 3);
+				check(PolygonVertexCount == 3);
 				TArray<FEdgeID> NewEdgeIDs;
-				const FPolygonID NewPolygonID = MeshDescription->CreatePolygon(PolygonGroupID, CornerInstanceIDs, &NewEdgeIDs);
+				const FTriangleID NewTriangleID = MeshDescription->CreateTriangle(PolygonGroupID, CornerInstanceIDs, &NewEdgeIDs);
 				check(NewEdgeIDs.Num() == 0);
+
+				for (int32 UVLayerIndex = 0; UVLayerIndex < FBXUVs.UniqueUVCount; UVLayerIndex++)
+				{
+					FUVID UVIDs[3] = {INDEX_NONE, INDEX_NONE, INDEX_NONE};
+					if (FBXUVs.LayerElementUV[UVLayerIndex] != NULL)
+					{
+						for (int32 VertexIndex = 0; VertexIndex < 3; ++VertexIndex)
+						{
+							int UVMapIndex = (FBXUVs.UVMappingMode[UVLayerIndex] == FbxLayerElement::eByControlPoint)
+								? MeshDescription->GetVertexInstanceVertex(CornerInstanceIDs[VertexIndex])
+								: SkippedVertexInstance + CurrentVertexInstanceIndex + VertexIndex;
+							int32 UVIndex = (FBXUVs.UVReferenceMode[UVLayerIndex] == FbxLayerElement::eDirect)
+								? UVMapIndex
+								: FBXUVs.LayerElementUV[UVLayerIndex]->GetIndexArray().GetAt(UVMapIndex);
+
+							UVIDs[VertexIndex] = UVIndex + UVOffsets[UVLayerIndex];
+
+							check(MeshDescription->VertexInstanceAttributes().GetAttribute<FVector2D>(CornerInstanceIDs[VertexIndex], MeshAttribute::VertexInstance::TextureCoordinate, UVLayerIndex) ==
+								MeshDescription->UVAttributes(UVLayerIndex).GetAttribute<FVector2D>(UVIndex + UVOffsets[UVLayerIndex], MeshAttribute::UV::UVCoordinate));
+						}
+					}
+
+					MeshDescription->SetTriangleUVIndices(NewTriangleID, UVIDs, UVLayerIndex);
+				}
+
+				CurrentVertexInstanceIndex += PolygonVertexCount;
 			}
+
 			//Call this after all GetMeshEdgeIndexForPolygon call this is for optimization purpose.
 			if (bBeginGetMeshEdgeIndexForPolygonCalled)
 			{
 				Mesh->EndGetMeshEdgeIndexForPolygon();
 			}
+
+			MeshDescription->ResumeVertexInstanceIndexing();
+			MeshDescription->ResumeEdgeIndexing();
+			MeshDescription->ResumePolygonIndexing();
+			MeshDescription->ResumePolygonGroupIndexing();
+			MeshDescription->ResumeUVIndexing();
+
 			if (SkippedVertexInstance > 0)
 			{
 				//We must compact the sparse array before reserving new space
@@ -1672,7 +1747,7 @@ UStaticMesh* UnFbx::FFbxImporter::ImportStaticMeshAsSingle(UObject* InParent, TA
 	if (bBuildStatus)
 	{
 		TVertexInstanceAttributesRef<FVector2D> VertexInstanceUVs = Attributes.GetVertexInstanceUVs();
-		int32 FirstOpenUVChannel = VertexInstanceUVs.GetNumIndices() >= MAX_MESH_TEXTURE_COORDS_MD ? 1 : VertexInstanceUVs.GetNumIndices();
+		int32 FirstOpenUVChannel = VertexInstanceUVs.GetNumChannels() >= MAX_MESH_TEXTURE_COORDS_MD ? 1 : VertexInstanceUVs.GetNumChannels();
 		TPolygonGroupAttributesRef<FName> PolygonGroupImportedMaterialSlotNames = Attributes.GetPolygonGroupMaterialSlotNames();
 
 		TArray<FStaticMaterial> MaterialToAdd;
diff --git a/Engine/Source/Editor/UnrealEd/Private/StaticMeshEdit.cpp b/Engine/Source/Editor/UnrealEd/Private/StaticMeshEdit.cpp
--- a/Engine/Source/Editor/UnrealEd/Private/StaticMeshEdit.cpp
+++ b/Engine/Source/Editor/UnrealEd/Private/StaticMeshEdit.cpp
@@ -688,18 +688,18 @@ inline bool FVerticesEqual(FVector& V1,FVector& V2)
 
 void GetBrushMesh(ABrush* Brush, UModel* Model, FMeshDescription& MeshDescription, TArray<FStaticMaterial>& OutMaterials)
 {
-	TVertexAttributesRef<FVector> VertexPositions = MeshDescription.VertexAttributes().GetAttributesRef<FVector>(MeshAttribute::Vertex::Position);
-	TVertexInstanceAttributesRef<FVector> VertexInstanceNormals = MeshDescription.VertexInstanceAttributes().GetAttributesRef<FVector>(MeshAttribute::VertexInstance::Normal);
-	TVertexInstanceAttributesRef<FVector> VertexInstanceTangents = MeshDescription.VertexInstanceAttributes().GetAttributesRef<FVector>(MeshAttribute::VertexInstance::Tangent);
-	TVertexInstanceAttributesRef<float> VertexInstanceBinormalSigns = MeshDescription.VertexInstanceAttributes().GetAttributesRef<float>(MeshAttribute::VertexInstance::BinormalSign);
-	TVertexInstanceAttributesRef<FVector4> VertexInstanceColors = MeshDescription.VertexInstanceAttributes().GetAttributesRef<FVector4>(MeshAttribute::VertexInstance::Color);
-	TVertexInstanceAttributesRef<FVector2D> VertexInstanceUVs = MeshDescription.VertexInstanceAttributes().GetAttributesRef<FVector2D>(MeshAttribute::VertexInstance::TextureCoordinate);
-	TEdgeAttributesRef<bool> EdgeHardnesses = MeshDescription.EdgeAttributes().GetAttributesRef<bool>(MeshAttribute::Edge::IsHard);
-	TEdgeAttributesRef<float> EdgeCreaseSharpnesses = MeshDescription.EdgeAttributes().GetAttributesRef<float>(MeshAttribute::Edge::CreaseSharpness);
-	TPolygonGroupAttributesRef<FName> PolygonGroupImportedMaterialSlotNames = MeshDescription.PolygonGroupAttributes().GetAttributesRef<FName>(MeshAttribute::PolygonGroup::ImportedMaterialSlotName);
+	FStaticMeshAttributes Attributes(MeshDescription);
+	TVertexAttributesRef<FVector> VertexPositions = Attributes.GetVertexPositions();
+	TVertexInstanceAttributesRef<FVector> VertexInstanceNormals = Attributes.GetVertexInstanceNormals();
+	TVertexInstanceAttributesRef<FVector> VertexInstanceTangents = Attributes.GetVertexInstanceTangents();
+	TVertexInstanceAttributesRef<float> VertexInstanceBinormalSigns = Attributes.GetVertexInstanceBinormalSigns();
+	TVertexInstanceAttributesRef<FVector4> VertexInstanceColors = Attributes.GetVertexInstanceColors();
+	TVertexInstanceAttributesRef<FVector2D> VertexInstanceUVs = Attributes.GetVertexInstanceUVs();
+	TEdgeAttributesRef<bool> EdgeHardnesses = Attributes.GetEdgeHardnesses();
+	TPolygonGroupAttributesRef<FName> PolygonGroupImportedMaterialSlotNames = Attributes.GetPolygonGroupMaterialSlotNames();
 	
 	//Make sure we have one UVChannel
-	VertexInstanceUVs.SetNumIndices(1);
+	VertexInstanceUVs.SetNumChannels(1);
 	
 	// Calculate the local to world transform for the source brush.
 	FMatrix	ActorToWorld = Brush ? Brush->ActorToWorld().ToMatrixWithScale() : FMatrix::Identity;
@@ -721,7 +721,7 @@ void GetBrushMesh(ABrush* Brush, UModel* Model, FMeshDescription& MeshDescriptio
 		}
 
 		int32 MaterialIndex = OutMaterials.AddUnique(FStaticMaterial(Material, Material->GetFName(), Material->GetFName()));
-		FPolygonGroupID CurrentPolygonGroupID = FPolygonGroupID::Invalid;
+		FPolygonGroupID CurrentPolygonGroupID = INDEX_NONE;
 		for (const FPolygonGroupID& PolygonGroupID : MeshDescription.PolygonGroups().GetElementIDs())
 		{
 			if (Material->GetFName() == PolygonGroupImportedMaterialSlotNames[PolygonGroupID])
@@ -730,7 +730,7 @@ void GetBrushMesh(ABrush* Brush, UModel* Model, FMeshDescription& MeshDescriptio
 				break;
 			}
 		}
-		if (CurrentPolygonGroupID == FPolygonGroupID::Invalid)
+		if (CurrentPolygonGroupID == INDEX_NONE)
 		{
 			CurrentPolygonGroupID = MeshDescription.CreatePolygonGroup();
 			PolygonGroupImportedMaterialSlotNames[CurrentPolygonGroupID] = Material->GetFName();
@@ -747,7 +747,7 @@ void GetBrushMesh(ABrush* Brush, UModel* Model, FMeshDescription& MeshDescriptio
 			Positions[ReverseVertices ? 0 : 2] = ActorToWorld.TransformPosition(Polygon.Vertices[0]) - PostSub;
 			Positions[1] = ActorToWorld.TransformPosition(Polygon.Vertices[VertexIndex - 1]) - PostSub;
 			Positions[ReverseVertices ? 2 : 0] = ActorToWorld.TransformPosition(Polygon.Vertices[VertexIndex]) - PostSub;
-			FVertexID VertexID[3] = { FVertexID::Invalid, FVertexID::Invalid, FVertexID::Invalid };
+			FVertexID VertexID[3] = { INDEX_NONE, INDEX_NONE, INDEX_NONE };
 			for (FVertexID IterVertexID : MeshDescription.Vertices().GetElementIDs())
 			{
 				if (FVerticesEqual(Positions[0], VertexPositions[IterVertexID]))
@@ -770,7 +770,7 @@ void GetBrushMesh(ABrush* Brush, UModel* Model, FMeshDescription& MeshDescriptio
 
 			for (int32 CornerIndex = 0; CornerIndex < 3; ++CornerIndex)
 			{
-				if (VertexID[CornerIndex] == FVertexID::Invalid)
+				if (VertexID[CornerIndex] == INDEX_NONE)
 				{
 					VertexID[CornerIndex] = MeshDescription.CreateVertex();
 					VertexPositions[VertexID[CornerIndex]] = Positions[CornerIndex];
diff --git a/Engine/Source/Runtime/Core/Public/UObject/ReleaseObjectVersion.h b/Engine/Source/Runtime/Core/Public/UObject/ReleaseObjectVersion.h
--- a/Engine/Source/Runtime/Core/Public/UObject/ReleaseObjectVersion.h
+++ b/Engine/Source/Runtime/Core/Public/UObject/ReleaseObjectVersion.h
@@ -102,6 +102,9 @@ struct CORE_API FReleaseObjectVersion
 		// FFieldPath will serialize the owner struct reference and only a short path to its property
 		FFieldPathOwnerSerialization,
 
+		// New MeshDescription format
+		MeshDescriptionNewFormat,
+
 		// -----<new versions can be added above this line>-------------------------------------------------
 		VersionPlusOne,
 		LatestVersion = VersionPlusOne - 1
diff --git a/Engine/Source/Runtime/Datasmith/DatasmithCore/Private/DatasmithUtils.cpp b/Engine/Source/Runtime/Datasmith/DatasmithCore/Private/DatasmithUtils.cpp
--- a/Engine/Source/Runtime/Datasmith/DatasmithCore/Private/DatasmithUtils.cpp
+++ b/Engine/Source/Runtime/Datasmith/DatasmithCore/Private/DatasmithUtils.cpp
@@ -313,7 +313,7 @@ bool FDatasmithMeshUtils::ToMeshDescription(FDatasmithMesh& DsMesh, FMeshDescrip
 	// At least one UV set must exist.
 	int32 DsUVCount = DsMesh.GetUVChannelsCount();
 	int32 MeshDescUVCount = FMath::Max(1, DsUVCount);
-	VertexInstanceUVs.SetNumIndices(MeshDescUVCount);
+	VertexInstanceUVs.SetNumChannels(MeshDescUVCount);
 
 	//Fill the vertex array
 	for (int32 VertexIndex = 0; VertexIndex < VertexCount; ++VertexIndex)
@@ -326,7 +326,7 @@ bool FDatasmithMeshUtils::ToMeshDescription(FDatasmithMesh& DsMesh, FMeshDescrip
 	auto GetOrCreatePolygonGroupId = [&](int32 MaterialIndex)
 	{
 		FPolygonGroupID& PolyGroupId = PolygonGroupMapping.FindOrAdd(MaterialIndex);
-		if (PolyGroupId == FPolygonGroupID::Invalid)
+		if (PolyGroupId == INDEX_NONE)
 		{
 			PolyGroupId = MeshDescription.CreatePolygonGroup();
 			FName ImportedSlotName = *FString::FromInt(MaterialIndex); // No access to DatasmithMeshHelper::DefaultSlotName
diff --git a/Engine/Source/Runtime/Engine/Private/StaticMesh.cpp b/Engine/Source/Runtime/Engine/Private/StaticMesh.cpp
--- a/Engine/Source/Runtime/Engine/Private/StaticMesh.cpp
+++ b/Engine/Source/Runtime/Engine/Private/StaticMesh.cpp
@@ -4138,11 +4138,11 @@ void UStaticMesh::FixupMaterialSlotName()
 	}
 }
 
-// If static mesh derived data needs to be rebuilt (new format, serialization
-// differences, etc.) replace the version GUID below with a new one.
+// If static mesh derived data needs to be rebuilt (MeshDescription new format,
+// serialization differences, etc.) replace the version GUID below with a new one.
 // In case of merge conflicts with DDC versions, you MUST generate a new GUID
 // and set this new GUID as the version.                                       
-#define MESHDATAKEY_STATICMESH_DERIVEDDATA_VER TEXT("940C97E946F545F5834146D0D58828BF")
+#define MESHDATAKEY_STATICMESH_DERIVEDDATA_VER TEXT("5D29A2C59A03418CBCE13C88747BA3AD")
 
 
 static const FString& GetMeshDataKeyStaticMeshDerivedDataVersion()
@@ -4395,7 +4395,7 @@ int32 UStaticMesh::GetNumUVChannels(int32 LODIndex)
 	if (MeshDescription)
 	{
 		FStaticMeshConstAttributes Attributes(*MeshDescription);
-		NumUVChannels = Attributes.GetVertexInstanceUVs().GetNumIndices();
+		NumUVChannels = Attributes.GetVertexInstanceUVs().GetNumChannels();
 	}
 #endif
 	return NumUVChannels;
@@ -5155,7 +5155,7 @@ void UStaticMesh::BuildFromMeshDescription(const FMeshDescription& MeshDescripti
 		StaticMeshVertex.TangentY = FVector::CrossProduct(VertexInstanceNormals[VertexInstanceID], VertexInstanceTangents[VertexInstanceID]).GetSafeNormal() * VertexInstanceBinormalSigns[VertexInstanceID];
 		StaticMeshVertex.TangentZ = VertexInstanceNormals[VertexInstanceID];
 
-		for (int32 UVIndex = 0; UVIndex < VertexInstanceUVs.GetNumIndices(); ++UVIndex)
+		for (int32 UVIndex = 0; UVIndex < VertexInstanceUVs.GetNumChannels(); ++UVIndex)
 		{
 			StaticMeshVertex.UVs[UVIndex] = VertexInstanceUVs.Get(VertexInstanceID, UVIndex);
 		}
@@ -5179,7 +5179,7 @@ void UStaticMesh::BuildFromMeshDescription(const FMeshDescription& MeshDescripti
 	}
 
 	LODResources.VertexBuffers.PositionVertexBuffer.Init(StaticMeshBuildVertices);
-	LODResources.VertexBuffers.StaticMeshVertexBuffer.Init(StaticMeshBuildVertices, VertexInstanceUVs.GetNumIndices());
+	LODResources.VertexBuffers.StaticMeshVertexBuffer.Init(StaticMeshBuildVertices, VertexInstanceUVs.GetNumChannels());
 
 	FColorVertexBuffer& ColorVertexBuffer = LODResources.VertexBuffers.ColorVertexBuffer;
 	if (bHasVertexColors)
@@ -6169,7 +6169,7 @@ void UStaticMesh::EnforceLightmapRestrictions(bool bUseRenderData)
 					// skip empty/stripped LODs
 					if (UVChannels.GetNumElements() > 0)
 					{
-						int NumChannelsInLOD = UVChannels.GetNumIndices();
+						int NumChannelsInLOD = UVChannels.GetNumChannels();
 						const FStaticMeshSourceModel& SourceModel = GetSourceModel(SourceLOD);
 
 						if (SourceModel.BuildSettings.bGenerateLightmapUVs)
diff --git a/Engine/Source/Runtime/Landscape/Private/LandscapeEdit.cpp b/Engine/Source/Runtime/Landscape/Private/LandscapeEdit.cpp
--- a/Engine/Source/Runtime/Landscape/Private/LandscapeEdit.cpp
+++ b/Engine/Source/Runtime/Landscape/Private/LandscapeEdit.cpp
@@ -3085,7 +3085,6 @@ bool ALandscapeProxy::ExportToRawMesh(int32 InExportLOD, FMeshDescription& OutRa
 
 	TVertexAttributesRef<FVector> VertexPositions = OutRawMesh.VertexAttributes().GetAttributesRef<FVector>(MeshAttribute::Vertex::Position);
 	TEdgeAttributesRef<bool> EdgeHardnesses = OutRawMesh.EdgeAttributes().GetAttributesRef<bool>(MeshAttribute::Edge::IsHard);
-	TEdgeAttributesRef<float> EdgeCreaseSharpnesses = OutRawMesh.EdgeAttributes().GetAttributesRef<float>(MeshAttribute::Edge::CreaseSharpness);
 	TPolygonGroupAttributesRef<FName> PolygonGroupImportedMaterialSlotNames = OutRawMesh.PolygonGroupAttributes().GetAttributesRef<FName>(MeshAttribute::PolygonGroup::ImportedMaterialSlotName);
 	TVertexInstanceAttributesRef<FVector> VertexInstanceNormals = OutRawMesh.VertexInstanceAttributes().GetAttributesRef<FVector>(MeshAttribute::VertexInstance::Normal);
 	TVertexInstanceAttributesRef<FVector> VertexInstanceTangents = OutRawMesh.VertexInstanceAttributes().GetAttributesRef<FVector>(MeshAttribute::VertexInstance::Tangent);
@@ -3093,9 +3092,9 @@ bool ALandscapeProxy::ExportToRawMesh(int32 InExportLOD, FMeshDescription& OutRa
 	TVertexInstanceAttributesRef<FVector4> VertexInstanceColors = OutRawMesh.VertexInstanceAttributes().GetAttributesRef<FVector4>(MeshAttribute::VertexInstance::Color);
 	TVertexInstanceAttributesRef<FVector2D> VertexInstanceUVs = OutRawMesh.VertexInstanceAttributes().GetAttributesRef<FVector2D>(MeshAttribute::VertexInstance::TextureCoordinate);
 
-	if (VertexInstanceUVs.GetNumIndices() < 2)
+	if (VertexInstanceUVs.GetNumChannels() < 2)
 	{
-		VertexInstanceUVs.SetNumIndices(2);
+		VertexInstanceUVs.SetNumChannels(2);
 	}
 
 	// User specified LOD to export
@@ -3131,7 +3130,7 @@ bool ALandscapeProxy::ExportToRawMesh(int32 InExportLOD, FMeshDescription& OutRa
 		OutRawMesh.ReserveNewVertexInstances(NumVertices);
 		OutRawMesh.ReserveNewEdges(NumVertices);
 
-		FPolygonGroupID PolygonGroupID = FPolygonGroupID::Invalid;
+		FPolygonGroupID PolygonGroupID = INDEX_NONE;
 		if (OutRawMesh.PolygonGroups().Num() < 1)
 		{
 			PolygonGroupID = OutRawMesh.CreatePolygonGroup();
@@ -3313,7 +3312,7 @@ bool ALandscapeProxy::ExportToRawMesh(int32 InExportLOD, FMeshDescription& OutRa
 								// Add lightmap UVs
 								VertexInstanceUVs.Set(VertexInstanceIDs[i], 1, UV);
 							}
-							auto AddTriangle = [&OutRawMesh, &EdgeHardnesses, &EdgeCreaseSharpnesses, &PolygonGroupID, &VertexIDs, &VertexInstanceIDs](int32 BaseIndex)
+							auto AddTriangle = [&OutRawMesh, &EdgeHardnesses, &PolygonGroupID, &VertexIDs, &VertexInstanceIDs](int32 BaseIndex)
 							{
 								//Create a polygon from this triangle
 								TArray<FVertexInstanceID> PerimeterVertexInstances;
@@ -3328,7 +3327,6 @@ bool ALandscapeProxy::ExportToRawMesh(int32 InExportLOD, FMeshDescription& OutRa
 								for (const FEdgeID NewEdgeID : NewEdgeIDs)
 								{
 									EdgeHardnesses[NewEdgeID] = false;
-									EdgeCreaseSharpnesses[NewEdgeID] = 0.0f;
 								}
 							};
 							AddTriangle(0);
diff --git a/Engine/Source/Runtime/MeshDescription/Private/MeshAttributeArray.cpp b/Engine/Source/Runtime/MeshDescription/Private/MeshAttributeArray.cpp
--- a/Engine/Source/Runtime/MeshDescription/Private/MeshAttributeArray.cpp
+++ b/Engine/Source/Runtime/MeshDescription/Private/MeshAttributeArray.cpp
@@ -2,23 +2,33 @@
 
 #include "MeshAttributeArray.h"
 #include "UObject/EditorObjectVersion.h"
+#include "UObject/ReleaseObjectVersion.h"
 
 
-FArchive& operator<<( FArchive& Ar, FAttributesSetEntry& Entry )
+FArchive& operator<<(FArchive& Ar, FAttributesSetEntry& Entry)
 {
-	if( Ar.IsLoading() )
+	if (Ar.IsLoading())
 	{
 		uint32 AttributeType;
 		Ar << AttributeType;
-		Entry.CreateArrayOfType( AttributeType );
-		Entry.Ptr->Serialize( Ar );
+
+		uint32 Extent = 1;
+		if (Ar.CustomVer(FEditorObjectVersion::GUID) >= FReleaseObjectVersion::MeshDescriptionNewFormat)
+		{
+			Ar << Extent;
+		}
+
+		Entry.CreateArrayOfType(AttributeType, Extent);
+		Entry.Ptr->Serialize(Ar);
 	}
 	else
 	{
-		check( Entry.Ptr.IsValid() );
+		check(Entry.Ptr.IsValid());
 		uint32 AttributeType = Entry.Ptr->GetType();
 		Ar << AttributeType;
-		Entry.Ptr->Serialize( Ar );
+		uint32 Extent = Entry.Ptr->GetExtent();
+		Ar << Extent;
+		Entry.Ptr->Serialize(Ar);
 	}
 
 	return Ar;
@@ -44,6 +54,7 @@ void SerializeLegacy( FArchive& Ar, FAttributesSetBase& AttributesSet )
 FArchive& operator<<( FArchive& Ar, FAttributesSetBase& AttributesSet )
 {
 	Ar.UsingCustomVersion( FEditorObjectVersion::GUID );
+	Ar.UsingCustomVersion( FReleaseObjectVersion::GUID );
 
 	if( Ar.IsLoading() && Ar.CustomVer( FEditorObjectVersion::GUID ) < FEditorObjectVersion::MeshDescriptionNewAttributeFormat )
 	{
@@ -111,3 +122,21 @@ void FAttributesSetBase::Remap( const TSparseArray<int32>& IndexRemap )
 		check( MapEntry.Value->GetNumElements() == NumElements );
 	}
 }
+
+
+void FAttributesSetBase::AppendAttributesFrom(const FAttributesSetBase& OtherAttributesSet)
+{
+	check(OtherAttributesSet.NumElements == NumElements);
+
+	for (const auto& MapPair : OtherAttributesSet.Map)
+	{
+		if (Map.Contains(MapPair.Key))
+		{
+			UE_LOG(LogMeshDescription, Log, TEXT("Appending attribute '%s' which already exists."), *MapPair.Key.ToString());
+		}
+		else
+		{
+			Map.Add(MapPair.Key, MapPair.Value);
+		}
+	}
+}
diff --git a/Engine/Source/Runtime/MeshDescription/Private/MeshAttributes.cpp b/Engine/Source/Runtime/MeshDescription/Private/MeshAttributes.cpp
--- a/Engine/Source/Runtime/MeshDescription/Private/MeshAttributes.cpp
+++ b/Engine/Source/Runtime/MeshDescription/Private/MeshAttributes.cpp
@@ -4,7 +4,40 @@
 
 namespace MeshAttribute
 {
-	const FName Vertex::Position("Position ");
+	namespace Vertex
+	{
+		const FName Position("Position ");
+	}
+
+	namespace VertexInstance
+	{
+		const FName VertexIndex("VertexIndex");
+	}
+
+	namespace Edge
+	{
+		const FName VertexIndex("VertexIndex");
+	}
+
+	namespace Triangle
+	{
+		const FName VertexInstanceIndex("VertexInstanceIndex");
+		const FName PolygonIndex("PolygonIndex");
+		const FName EdgeIndex("EdgeIndex");
+		const FName VertexIndex("VertexIndex");
+		const FName UVIndex("UVIndex");
+		const FName PolygonGroupIndex("PolygonGroupIndex");
+	}
+
+	namespace UV
+	{
+		const FName UVCoordinate("UVCoordinate");
+	}
+
+	namespace Polygon
+	{
+		const FName PolygonGroupIndex("PolygonGroupIndex");
+	}
 }
 
 
