Fully deprecated old MeshDescription APIs.
Prepared direct attribute access for deprecation, preferring use of APIs to access static mesh attributes.
Fixed recently merged Enterprise code to comply with new APIs.
Changed all tools to use triangle-centric iteration where possible.
Added new MeshAttributeArray APIs for querying attribute flags, and added a new Mandatory flag.
Various bug fixes.

diff --git a/Engine/Plugins/Editor/MeshEditor/Source/MeshEditor/MeshEditorGeometryCollectionAdapter.cpp b/Engine/Plugins/Editor/MeshEditor/Source/MeshEditor/MeshEditorGeometryCollectionAdapter.cpp
--- a/Engine/Plugins/Editor/MeshEditor/Source/MeshEditor/MeshEditorGeometryCollectionAdapter.cpp
+++ b/Engine/Plugins/Editor/MeshEditor/Source/MeshEditor/MeshEditorGeometryCollectionAdapter.cpp
@@ -125,7 +125,7 @@ void UMeshEditorGeometryCollectionAdapter::AddPolygonGroupToWireframe(const FMes
 			check(GeometryCollection);
 			FColor PolygonGroupColour = FColor::Blue;
 
-			TArrayView<const FPolygonID> PolygonGroupIDs = MeshDescription->GetPolygonGroupPolygons(FPolygonGroupID(SelectedPolygonGroup));
+			TArrayView<const FPolygonID> PolygonGroupIDs = MeshDescription->GetPolygonGroupPolygonIDs(FPolygonGroupID(SelectedPolygonGroup));
 			if (PolygonGroupIDs.Num() > 0)
 			{
 				// Add all polygons and edge instances
diff --git a/Engine/Plugins/Editor/MeshEditor/Source/MeshEditor/MeshEditorMode.cpp b/Engine/Plugins/Editor/MeshEditor/Source/MeshEditor/MeshEditorMode.cpp
--- a/Engine/Plugins/Editor/MeshEditor/Source/MeshEditor/MeshEditorMode.cpp
+++ b/Engine/Plugins/Editor/MeshEditor/Source/MeshEditor/MeshEditorMode.cpp
@@ -2492,7 +2492,7 @@ void FMeshEditorMode::AddMeshElementToOverlay( UOverlayComponent* OverlayCompone
 					case EEditableMeshElementType::Polygon:
 					{
 						const FPolygonID PolygonID( MeshElement.ElementAddress.ElementID );
-						TArrayView<const FTriangleID> TriangleIDs = MeshDescription->GetPolygonTriangleIDs( PolygonID );
+						TArrayView<const FTriangleID> TriangleIDs = MeshDescription->GetPolygonTriangles( PolygonID );
 
 						for( int32 PolygonTriangle = 0; PolygonTriangle < TriangleIDs.Num(); PolygonTriangle++ )
 						{
diff --git a/Engine/Plugins/Editor/MeshEditor/Source/PolygonModeling/SplitMesh.cpp b/Engine/Plugins/Editor/MeshEditor/Source/PolygonModeling/SplitMesh.cpp
--- a/Engine/Plugins/Editor/MeshEditor/Source/PolygonModeling/SplitMesh.cpp
+++ b/Engine/Plugins/Editor/MeshEditor/Source/PolygonModeling/SplitMesh.cpp
@@ -224,11 +224,11 @@ void USplitMeshCommand::Execute(IMeshEditorModeEditingContract& MeshEditorMode)
 		for(const FEdgeID EdgeId : BoundaryIds)
 		{
 			FVertexID Vertex0 = Mesh->GetMeshDescription()->GetEdgeVertex(EdgeId, 0);
-			TArrayView<const FVertexInstanceID> VertexInstances0 = Mesh->GetMeshDescription()->GetVertexVertexInstances(Vertex0);
+			TArrayView<const FVertexInstanceID> VertexInstances0 = Mesh->GetMeshDescription()->GetVertexVertexInstanceIDs(Vertex0);
 			check(VertexInstances0.Num() > 0);
 
 			FVertexID Vertex1 = Mesh->GetMeshDescription()->GetEdgeVertex(EdgeId, 1);
-			TArrayView<const FVertexInstanceID> VertexInstances1 = Mesh->GetMeshDescription()->GetVertexVertexInstances(Vertex1);
+			TArrayView<const FVertexInstanceID> VertexInstances1 = Mesh->GetMeshDescription()->GetVertexVertexInstanceIDs(Vertex1);
 			check(VertexInstances1.Num() > 0);
 
 			TArray<FVertexInstanceID> PolygonVertexInstances;
diff --git a/Engine/Plugins/Enterprise/DatasmithCADImporter/Source/CADLibrary/Private/CoreTechHelper.cpp b/Engine/Plugins/Enterprise/DatasmithCADImporter/Source/CADLibrary/Private/CoreTechHelper.cpp
--- a/Engine/Plugins/Enterprise/DatasmithCADImporter/Source/CADLibrary/Private/CoreTechHelper.cpp
+++ b/Engine/Plugins/Enterprise/DatasmithCADImporter/Source/CADLibrary/Private/CoreTechHelper.cpp
@@ -80,7 +80,7 @@ void FillVertexPosition(const FImportParameters& ImportParams, const FMeshParame
 	BBox.Min *= ImportParams.ScaleFactor;
 	BBox.Max *= ImportParams.ScaleFactor;
 
-	TVertexAttributesRef<FVector> VertexPositions = MeshDescription.VertexAttributes().GetAttributesRef<FVector>(MeshAttribute::Vertex::Position);
+	TVertexAttributesRef<FVector> VertexPositions = MeshDescription.GetVertexPositions();
 
 	// Create a list of vertex Z/index pairs
 	TArray<FVertexData> VertexDataSet;
diff --git a/Engine/Plugins/Enterprise/DatasmithCADImporter/Source/CADLibrary/Private/MeshEditingWrapper.cpp b/Engine/Plugins/Enterprise/DatasmithCADImporter/Source/CADLibrary/Private/MeshEditingWrapper.cpp
--- a/Engine/Plugins/Enterprise/DatasmithCADImporter/Source/CADLibrary/Private/MeshEditingWrapper.cpp
+++ b/Engine/Plugins/Enterprise/DatasmithCADImporter/Source/CADLibrary/Private/MeshEditingWrapper.cpp
@@ -41,7 +41,7 @@ void FMeshEditingWrapper::DefineEdgeTopology(FEdgeID EdgeID)
 		return;
 	}
 
-	TArrayView<const FTriangleID> EdgeConnectedPolygons = MeshDescription.GetEdgeConnectedTriangles(EdgeID);
+	TArrayView<const FTriangleID> EdgeConnectedPolygons = MeshDescription.GetEdgeConnectedTriangleIDs(EdgeID);
 
 	switch (EdgeConnectedPolygons.Num())
 	{
@@ -67,7 +67,7 @@ void FMeshEditingWrapper::DefineVertexTopologyApproximation(FVertexID VertexID)
 		return;
 	}
 
-	TArrayView<const FEdgeID> VertexConnectedEdgeIDs = MeshDescription.GetVertexConnectedEdges(VertexID);
+	TArrayView<const FEdgeID> VertexConnectedEdgeIDs = MeshDescription.GetVertexConnectedEdgeIDs(VertexID);
 
 	switch (VertexConnectedEdgeIDs.Num())
 	{
@@ -118,7 +118,7 @@ void FMeshEditingWrapper::DefineVertexTopologyApproximation(FVertexID VertexID)
 
 			do
 			{
-				TArrayView<const FTriangleID> EdgeConnectedPolygons = MeshDescription.GetEdgeConnectedTriangles(EdgeID);
+				TArrayView<const FTriangleID> EdgeConnectedPolygons = MeshDescription.GetEdgeConnectedTriangleIDs(EdgeID);
 				// Border edge no more triangles to process, exit from the loop
 				if (EdgeConnectedPolygons.Num() < 2)
 				{
@@ -211,7 +211,7 @@ void FMeshEditingWrapper::GetTriangleBoundingBox(FTriangleID Triangle, FVector&
 void FMeshEditingWrapper::GetVertexBoundingBox(FVertexInstanceID VertexInstanceID, FVector& MinCorner, FVector& MaxCorner, FVertexInstanceID HighestVertex[3], FVertexInstanceID LowestVertex[3]) const
 {
 	FVertexID VertexID = MeshDescription.GetVertexInstanceVertex(VertexInstanceID);
-	const FVector VertexPosition = MeshDescription.VertexAttributes().GetAttribute<FVector>(VertexID, MeshAttribute::Vertex::Position);
+	const FVector VertexPosition = MeshDescription.GetVertexPositions()[VertexID];
 
 	if (MaxCorner[0] < VertexPosition[0])
 	{
@@ -259,13 +259,13 @@ void FMeshEditingWrapper::SwapTriangleOrientation(FTriangleID Triangle)
 // Edge
 const FTriangleID FMeshEditingWrapper::GetOtherTriangleAtEdge(FEdgeID EdgeID, FTriangleID Triangle) const
 {
-	TArrayView<const FTriangleID> EdgeConnectedPolygons = MeshDescription.GetEdgeConnectedTriangles(EdgeID);
+	TArrayView<const FTriangleID> EdgeConnectedPolygons = MeshDescription.GetEdgeConnectedTriangleIDs(EdgeID);
 	return EdgeConnectedPolygons.Num() < 2 ? INDEX_NONE : (EdgeConnectedPolygons[0] == Triangle ? EdgeConnectedPolygons[1] : EdgeConnectedPolygons[0]);
 }
 
 bool FMeshEditingWrapper::GetEdgeDirectionInTriangle(FEdgeID EdgeID, int32 TriangleIndex) const
 {
-	TArrayView<const FTriangleID> EdgeConnectedTriangles = MeshDescription.GetEdgeConnectedTriangles(EdgeID);
+	TArrayView<const FTriangleID> EdgeConnectedTriangles = MeshDescription.GetEdgeConnectedTriangleIDs(EdgeID);
 
 	if (EdgeConnectedTriangles.Num() > TriangleIndex)
 	{
diff --git a/Engine/Plugins/Enterprise/DatasmithCADImporter/Source/CADLibrary/Private/MeshOperator.cpp b/Engine/Plugins/Enterprise/DatasmithCADImporter/Source/CADLibrary/Private/MeshOperator.cpp
--- a/Engine/Plugins/Enterprise/DatasmithCADImporter/Source/CADLibrary/Private/MeshOperator.cpp
+++ b/Engine/Plugins/Enterprise/DatasmithCADImporter/Source/CADLibrary/Private/MeshOperator.cpp
@@ -123,7 +123,7 @@ bool MeshOperator::OrientMesh(FMeshDescription& MeshDescription)
 					if (MeshWrapper.IsVertexOfCategory(HighestVertex[VertexIndex], EElementCategory::ElementCategorySurface))
 					{
 						FVertexID VertexID = MeshDescription.GetVertexInstanceVertex(HighestVertex[VertexIndex]);
-						TArrayView<const FVertexInstanceID> CoincidentVertexInstanceIdSet = MeshDescription.GetVertexVertexInstances(VertexID);
+						TArrayView<const FVertexInstanceID> CoincidentVertexInstanceIdSet = MeshDescription.GetVertexVertexInstanceIDs(VertexID);
 						float MaxComponent = 0;
 						for (const FVertexInstanceID VertexInstanceID : CoincidentVertexInstanceIdSet)
 						{
@@ -147,7 +147,7 @@ bool MeshOperator::OrientMesh(FMeshDescription& MeshDescription)
 					if (MeshWrapper.IsVertexOfCategory(LowestVertex[VertexIndex], EElementCategory::ElementCategorySurface))
 					{
 						FVertexID VertexID = MeshDescription.GetVertexInstanceVertex(LowestVertex[VertexIndex]);
-						TArrayView<const FVertexInstanceID> CoincidentVertexInstanceIdSet = MeshDescription.GetVertexVertexInstances(VertexID);
+						TArrayView<const FVertexInstanceID> CoincidentVertexInstanceIdSet = MeshDescription.GetVertexVertexInstanceIDs(VertexID);
 						float MaxComponent = 0;
 						for (const FVertexInstanceID VertexInstanceID : CoincidentVertexInstanceIdSet)
 						{
diff --git a/Engine/Plugins/Enterprise/DatasmithCADImporter/Source/DatasmithCoreTechExtension/Private/CoreTechBlueprintLibrary.cpp b/Engine/Plugins/Enterprise/DatasmithCADImporter/Source/DatasmithCoreTechExtension/Private/CoreTechBlueprintLibrary.cpp
--- a/Engine/Plugins/Enterprise/DatasmithCADImporter/Source/DatasmithCoreTechExtension/Private/CoreTechBlueprintLibrary.cpp
+++ b/Engine/Plugins/Enterprise/DatasmithCADImporter/Source/DatasmithCoreTechExtension/Private/CoreTechBlueprintLibrary.cpp
@@ -47,10 +47,10 @@ bool UCoreTechBlueprintLibrary::RetessellateStaticMeshWithNotification(UStaticMe
 				StaticMesh->PreEditChange( nullptr );
 			}
 
-			const int32 OldNumberOfUVChannels = DestinationMeshDescription->VertexInstanceAttributes().GetAttributeChannelCount(MeshAttribute::VertexInstance::TextureCoordinate);
+			const int32 OldNumberOfUVChannels = FStaticMeshAttributes(*DestinationMeshDescription).GetVertexInstanceUVs().GetNumChannels();
 			if (FCoreTechRetessellate_Impl::ApplyOnOneAsset( *StaticMesh, *CoreTechData, TessellationSettings ))
 			{
-				const int32 NumberOfUVChannels = DestinationMeshDescription->VertexInstanceAttributes().GetAttributeChannelCount(MeshAttribute::VertexInstance::TextureCoordinate);
+				const int32 NumberOfUVChannels = FStaticMeshAttributes(*DestinationMeshDescription).GetVertexInstanceUVs().GetNumChannels();
 				if (NumberOfUVChannels < OldNumberOfUVChannels)
 				{
 					FailureReason = FText::Format(NSLOCTEXT("BlueprintRetessellation", "UVChannelsDestroyed", "Tessellation operation on Static Mesh {0} is destroying all UV channels above channel #{1}"), FText::FromString(StaticMesh->GetName()), NumberOfUVChannels - 1);
diff --git a/Engine/Plugins/Enterprise/DatasmithCADImporter/Source/DatasmithWireTranslator/Private/OpenModelUtils.cpp b/Engine/Plugins/Enterprise/DatasmithCADImporter/Source/DatasmithWireTranslator/Private/OpenModelUtils.cpp
--- a/Engine/Plugins/Enterprise/DatasmithCADImporter/Source/DatasmithWireTranslator/Private/OpenModelUtils.cpp
+++ b/Engine/Plugins/Enterprise/DatasmithCADImporter/Source/DatasmithWireTranslator/Private/OpenModelUtils.cpp
@@ -1666,8 +1666,7 @@ bool OpenModelUtils::TransferAlMeshToMeshDescription(const AlMesh& AliasMesh, FM
 	TVertexInstanceAttributesRef<FVector> VertexInstanceNormals = Attributes.GetVertexInstanceNormals();
 	TVertexInstanceAttributesRef<FVector2D> VertexInstanceUVs = Attributes.GetVertexInstanceUVs();
 	TPolygonGroupAttributesRef<FName> PolygonGroupImportedMaterialSlotNames = Attributes.GetPolygonGroupMaterialSlotNames();
-
-	TVertexAttributesRef<FVector> VertexPositions = MeshDescription.VertexAttributes().GetAttributesRef<FVector>(MeshAttribute::Vertex::Position);
+	TVertexAttributesRef<FVector> VertexPositions = MeshDescription.GetVertexPositions();
 
 	// Prepared for static mesh usage ?
 	if (!VertexPositions.IsValid() || !VertexInstanceNormals.IsValid() || !VertexInstanceUVs.IsValid() || !PolygonGroupImportedMaterialSlotNames.IsValid())
diff --git a/Engine/Plugins/Enterprise/DatasmithFBXImporter/Source/DatasmithFBXTranslator/Private/DatasmithFBXFileImporter.cpp b/Engine/Plugins/Enterprise/DatasmithFBXImporter/Source/DatasmithFBXTranslator/Private/DatasmithFBXFileImporter.cpp
--- a/Engine/Plugins/Enterprise/DatasmithFBXImporter/Source/DatasmithFBXTranslator/Private/DatasmithFBXFileImporter.cpp
+++ b/Engine/Plugins/Enterprise/DatasmithFBXImporter/Source/DatasmithFBXTranslator/Private/DatasmithFBXFileImporter.cpp
@@ -1024,16 +1024,15 @@ void FDatasmithFBXFileImporter::DoImportMesh(FbxMesh* InMesh, FDatasmithFBXScene
 
 	FStaticMeshAttributes StaticMeshAttributes{ MeshDescription };
 	StaticMeshAttributes.Register();
-	StaticMeshAttributes.RegisterTriangleNormalAndTangentAttributes();	// @todo: remove if not necessary
-
-	TVertexAttributesRef<FVector> VertexPositions = MeshDescription.VertexAttributes().GetAttributesRef<FVector>(MeshAttribute::Vertex::Position);
-	TVertexInstanceAttributesRef<FVector> VertexInstanceNormals = MeshDescription.VertexInstanceAttributes().GetAttributesRef<FVector>(MeshAttribute::VertexInstance::Normal);
-	TVertexInstanceAttributesRef<FVector> VertexInstanceTangents = MeshDescription.VertexInstanceAttributes().GetAttributesRef<FVector>(MeshAttribute::VertexInstance::Tangent);
-	TVertexInstanceAttributesRef<float> VertexInstanceBinormalSigns = MeshDescription.VertexInstanceAttributes().GetAttributesRef<float>(MeshAttribute::VertexInstance::BinormalSign);
-	TVertexInstanceAttributesRef<FVector4> VertexInstanceColors = MeshDescription.VertexInstanceAttributes().GetAttributesRef<FVector4>(MeshAttribute::VertexInstance::Color);
-	TVertexInstanceAttributesRef<FVector2D> VertexInstanceUVs = MeshDescription.VertexInstanceAttributes().GetAttributesRef<FVector2D>(MeshAttribute::VertexInstance::TextureCoordinate);
-	TEdgeAttributesRef<bool> EdgeHardnesses = MeshDescription.EdgeAttributes().GetAttributesRef<bool>(MeshAttribute::Edge::IsHard);
-	TPolygonGroupAttributesRef<FName> PolygonGroupImportedMaterialSlotNames = MeshDescription.PolygonGroupAttributes().GetAttributesRef<FName>(MeshAttribute::PolygonGroup::ImportedMaterialSlotName);
+
+	TVertexAttributesRef<FVector> VertexPositions = StaticMeshAttributes.GetVertexPositions();
+	TVertexInstanceAttributesRef<FVector> VertexInstanceNormals = StaticMeshAttributes.GetVertexInstanceNormals();
+	TVertexInstanceAttributesRef<FVector> VertexInstanceTangents = StaticMeshAttributes.GetVertexInstanceTangents();
+	TVertexInstanceAttributesRef<float> VertexInstanceBinormalSigns = StaticMeshAttributes.GetVertexInstanceBinormalSigns();
+	TVertexInstanceAttributesRef<FVector4> VertexInstanceColors = StaticMeshAttributes.GetVertexInstanceColors();
+	TVertexInstanceAttributesRef<FVector2D> VertexInstanceUVs = StaticMeshAttributes.GetVertexInstanceUVs();
+	TEdgeAttributesRef<bool> EdgeHardnesses = StaticMeshAttributes.GetEdgeHardnesses();
+	TPolygonGroupAttributesRef<FName> PolygonGroupImportedMaterialSlotNames = StaticMeshAttributes.GetPolygonGroupMaterialSlotNames();
 
 	// Reserve space for attributes.
 	MeshDescription.ReserveNewVertices(VertexCount);
diff --git a/Engine/Plugins/Enterprise/DatasmithFBXImporter/Source/DatasmithFBXTranslator/Private/DatasmithFBXScene.cpp b/Engine/Plugins/Enterprise/DatasmithFBXImporter/Source/DatasmithFBXTranslator/Private/DatasmithFBXScene.cpp
--- a/Engine/Plugins/Enterprise/DatasmithFBXImporter/Source/DatasmithFBXTranslator/Private/DatasmithFBXScene.cpp
+++ b/Engine/Plugins/Enterprise/DatasmithFBXImporter/Source/DatasmithFBXTranslator/Private/DatasmithFBXScene.cpp
@@ -38,14 +38,14 @@ const FMD5Hash& FDatasmithFBXSceneMesh::GetHash()
 
 bool FDatasmithFBXSceneMesh::HasNormals() const
 {
-	const TVertexInstanceAttributesConstRef<FVector> Normals = MeshDescription.VertexInstanceAttributes().GetAttributesRef<FVector>(MeshAttribute::VertexInstance::Normal);
-	return Normals.GetNumElements() > 0 && Normals[MeshDescription.VertexInstances().GetFirstValidID()].SizeSquared() > 0;
+	TVertexInstanceAttributesConstRef<FVector> Normals = FStaticMeshConstAttributes(MeshDescription).GetVertexInstanceNormals();
+	return Normals.IsValid() && Normals.GetNumElements() > 0 && Normals[MeshDescription.VertexInstances().GetFirstValidID()].SizeSquared() > 0;
 }
 
 bool FDatasmithFBXSceneMesh::HasTangents() const
 {
-	const TVertexInstanceAttributesConstRef<FVector> Tangents = MeshDescription.VertexInstanceAttributes().GetAttributesRef<FVector>(MeshAttribute::VertexInstance::Tangent);
-	return Tangents.GetNumElements() > 0 && Tangents[MeshDescription.VertexInstances().GetFirstValidID()].SizeSquared() > 0;
+	TVertexInstanceAttributesConstRef<FVector> Tangents = FStaticMeshConstAttributes(MeshDescription).GetVertexInstanceTangents();
+	return Tangents.IsValid() && Tangents.GetNumElements() > 0 && Tangents[MeshDescription.VertexInstances().GetFirstValidID()].SizeSquared() > 0;
 }
 
 FDatasmithFBXSceneCamera::FDatasmithFBXSceneCamera()
diff --git a/Engine/Plugins/Enterprise/DatasmithIFCImporter/Source/DatasmithIFCTranslator/Private/IFC/IFCStaticMeshFactory.cpp b/Engine/Plugins/Enterprise/DatasmithIFCImporter/Source/DatasmithIFCTranslator/Private/IFC/IFCStaticMeshFactory.cpp
--- a/Engine/Plugins/Enterprise/DatasmithIFCImporter/Source/DatasmithIFCTranslator/Private/IFC/IFCStaticMeshFactory.cpp
+++ b/Engine/Plugins/Enterprise/DatasmithIFCImporter/Source/DatasmithIFCTranslator/Private/IFC/IFCStaticMeshFactory.cpp
@@ -41,10 +41,11 @@ namespace IFC
 
 	void FStaticMeshFactory::FillMeshDescription(const IFC::FObject* InObject, FMeshDescription* MeshDescription) const
 	{
-		TVertexAttributesRef<FVector> VertexPositions = MeshDescription->VertexAttributes().GetAttributesRef<FVector>(MeshAttribute::Vertex::Position);
-		TPolygonGroupAttributesRef<FName> PolygonGroupImportedMaterialSlotNames = MeshDescription->PolygonGroupAttributes().GetAttributesRef<FName>(MeshAttribute::PolygonGroup::ImportedMaterialSlotName);
-		TVertexInstanceAttributesRef<FVector> VertexInstanceNormals = MeshDescription->VertexInstanceAttributes().GetAttributesRef<FVector>(MeshAttribute::VertexInstance::Normal);
-		TVertexInstanceAttributesRef<FVector2D> VertexInstanceUVs = MeshDescription->VertexInstanceAttributes().GetAttributesRef<FVector2D>(MeshAttribute::VertexInstance::TextureCoordinate);
+		FStaticMeshAttributes Attributes(*MeshDescription);
+		TVertexAttributesRef<FVector> VertexPositions = Attributes.GetVertexPositions();
+		TPolygonGroupAttributesRef<FName> PolygonGroupImportedMaterialSlotNames = Attributes.GetPolygonGroupMaterialSlotNames();
+		TVertexInstanceAttributesRef<FVector> VertexInstanceNormals = Attributes.GetVertexInstanceNormals();
+		TVertexInstanceAttributesRef<FVector2D> VertexInstanceUVs = Attributes.GetVertexInstanceUVs();
 
 		const int32 NumUVs = 1;
 		VertexInstanceUVs.SetNumChannels(NumUVs);
diff --git a/Engine/Plugins/Enterprise/DatasmithImporter/Source/DatasmithImporter/Private/DatasmithStaticMeshImporter.cpp b/Engine/Plugins/Enterprise/DatasmithImporter/Source/DatasmithImporter/Private/DatasmithStaticMeshImporter.cpp
--- a/Engine/Plugins/Enterprise/DatasmithImporter/Source/DatasmithImporter/Private/DatasmithStaticMeshImporter.cpp
+++ b/Engine/Plugins/Enterprise/DatasmithImporter/Source/DatasmithImporter/Private/DatasmithStaticMeshImporter.cpp
@@ -234,14 +234,14 @@ UStaticMesh* FDatasmithStaticMeshImporter::ImportStaticMesh(const TSharedRef< ID
 
 bool FDatasmithStaticMeshImporter::ShouldRecomputeNormals(const FMeshDescription& MeshDescription, int32 BuildRequirements)
 {
-	const TVertexInstanceAttributesConstRef<FVector> Normals = MeshDescription.VertexInstanceAttributes().GetAttributesRef<FVector>(MeshAttribute::VertexInstance::Normal);
+	TVertexInstanceAttributesConstRef<FVector> Normals = FStaticMeshConstAttributes(MeshDescription).GetVertexInstanceNormals();
 	check(Normals.IsValid());
 	return Algo::AnyOf(MeshDescription.VertexInstances().GetElementIDs(), [&](const FVertexInstanceID& InstanceID) { return !Normals[InstanceID].IsNormalized(); });
 }
 
 bool FDatasmithStaticMeshImporter::ShouldRecomputeTangents(const FMeshDescription& MeshDescription, int32 BuildRequirements)
 {
-	const TVertexInstanceAttributesConstRef<FVector> Tangents = MeshDescription.VertexInstanceAttributes().GetAttributesRef<FVector>(MeshAttribute::VertexInstance::Tangent);
+	TVertexInstanceAttributesConstRef<FVector> Tangents = FStaticMeshConstAttributes(MeshDescription).GetVertexInstanceTangents();
 	check(Tangents.IsValid());
 	return Algo::AnyOf(MeshDescription.VertexInstances().GetElementIDs(), [&](const FVertexInstanceID& InstanceID) { return !Tangents[InstanceID].IsNormalized(); });
 }
@@ -488,7 +488,7 @@ public:
 					ThisSection.SectionIndex = SectionIndex++;
 					ThisSection.SectionId = PolygonGroupID;
 					ThisSection.MaterialSlotName = MaterialSlotNameAttribute[PolygonGroupID];
-					ThisSection.PolyCount = MeshDescription->GetPolygonGroupPolygons(PolygonGroupID).Num();
+					ThisSection.PolyCount = MeshDescription->GetPolygonGroupPolygonIDs(PolygonGroupID).Num();
 				}
 			}
 		}
diff --git a/Engine/Plugins/Enterprise/DatasmithImporter/Source/DatasmithTranslator/Private/Utility/DatasmithMeshHelper.cpp b/Engine/Plugins/Enterprise/DatasmithImporter/Source/DatasmithTranslator/Private/Utility/DatasmithMeshHelper.cpp
--- a/Engine/Plugins/Enterprise/DatasmithImporter/Source/DatasmithTranslator/Private/Utility/DatasmithMeshHelper.cpp
+++ b/Engine/Plugins/Enterprise/DatasmithImporter/Source/DatasmithTranslator/Private/Utility/DatasmithMeshHelper.cpp
@@ -47,7 +47,7 @@ namespace DatasmithMeshHelper
 
 	bool IsTriangleDegenerated(const FMeshDescription& Mesh, const FTriangleID TriangleID)
 	{
-		const TVertexAttributesConstRef<FVector> VertexPositions = Mesh.VertexAttributes().GetAttributesRef<FVector>(MeshAttribute::Vertex::Position);
+		TVertexAttributesConstRef<FVector> VertexPositions = Mesh.GetVertexPositions();
 		TArrayView<const FVertexID> TriangleVerts = Mesh.GetTriangleVertices(TriangleID);
 		float NormalLengthSquare = ((VertexPositions[TriangleVerts[1]] - VertexPositions[TriangleVerts[2]]) ^ (VertexPositions[TriangleVerts[0]] - VertexPositions[TriangleVerts[2]])).SizeSquared();
 		return NormalLengthSquare < SMALL_NUMBER;
@@ -225,7 +225,7 @@ namespace DatasmithMeshHelper
 
 	int32 GetNumUVChannel(const FMeshDescription& Mesh)
 	{
-		const TVertexInstanceAttributesConstRef<FVector2D> UVChannels = Mesh.VertexInstanceAttributes().GetAttributesRef<FVector2D>(MeshAttribute::VertexInstance::TextureCoordinate);
+		TVertexInstanceAttributesConstRef<FVector2D> UVChannels = FStaticMeshConstAttributes(Mesh).GetVertexInstanceUVs();
 		return UVChannels.GetNumChannels();
 	}
 
@@ -238,7 +238,7 @@ namespace DatasmithMeshHelper
 	{
 		if (HasUVChannel(Mesh, ChannelIndex))
 		{
-			const TVertexInstanceAttributesConstRef<FVector2D> UVChannels = Mesh.VertexInstanceAttributes().GetAttributesRef<FVector2D>(MeshAttribute::VertexInstance::TextureCoordinate);
+			TVertexInstanceAttributesConstRef<FVector2D> UVChannels = FStaticMeshConstAttributes(Mesh).GetVertexInstanceUVs();
 			const FVector2D DefValue = UVChannels.GetDefaultValue();
 			for (FVertexInstanceID InstanceID : Mesh.VertexInstances().GetElementIDs())
 			{
@@ -260,7 +260,7 @@ namespace DatasmithMeshHelper
 
 		if (!HasUVChannel(Mesh, ChannelIndex))
 		{
-			TVertexInstanceAttributesRef<FVector2D> UVChannels = Mesh.VertexInstanceAttributes().GetAttributesRef<FVector2D>(MeshAttribute::VertexInstance::TextureCoordinate);
+			TVertexInstanceAttributesRef<FVector2D> UVChannels = FStaticMeshAttributes(Mesh).GetVertexInstanceUVs();
 			UVChannels.SetNumChannels(ChannelIndex + 1);
 		}
 		return true;
@@ -320,7 +320,7 @@ namespace DatasmithMeshHelper
 		TMap<FVertexInstanceID, FVector2D> TexCoords;
 		FStaticMeshOperations::GenerateBoxUV(MeshDescription, UVParameters, TexCoords);
 
-		TMeshAttributesRef<FVertexInstanceID, FVector2D> UVs = MeshDescription.VertexInstanceAttributes().GetAttributesRef<FVector2D>(MeshAttribute::VertexInstance::TextureCoordinate);
+		TVertexInstanceAttributesRef<FVector2D> UVs = FStaticMeshAttributes(MeshDescription).GetVertexInstanceUVs();
 		if (UVs.GetNumChannels() == 0)
 		{
 			UVs.SetNumChannels(1);
@@ -341,30 +341,28 @@ namespace DatasmithMeshHelper
 
 	bool IsMeshValid(const FMeshDescription& Mesh, FVector BuildScale)
 	{
-		const TVertexAttributesConstRef<FVector> VertexPositions = Mesh.VertexAttributes().GetAttributesRef<FVector>(MeshAttribute::Vertex::Position);
+		TVertexAttributesConstRef<FVector> VertexPositions = Mesh.GetVertexPositions();
 		FVector RawNormalScale(BuildScale.Y*BuildScale.Z, BuildScale.X*BuildScale.Z, BuildScale.X*BuildScale.Y); // Component-wise scale
-		for (const FPolygonID PolygonID : Mesh.Polygons().GetElementIDs())
+
+		for (const FTriangleID TriangleID : Mesh.Triangles().GetElementIDs())
 		{
-			FVector Corners[3];
-			for (const FTriangleID TriangleID : Mesh.GetPolygonTriangleIDs(PolygonID))
+			TArrayView<const FVertexID> VertexIDs = Mesh.GetTriangleVertices(TriangleID);
+			FVector Corners[3] =
 			{
-				for (int32 CornerIndex = 0; CornerIndex < 3; ++CornerIndex)
-				{
-					FVertexInstanceID VertexInstanceID = Mesh.GetTriangleVertexInstance(TriangleID, CornerIndex);
-					FVertexID VertexID = Mesh.GetVertexInstanceVertex(VertexInstanceID);
-					Corners[CornerIndex] = VertexPositions[VertexID];
-				}
-
-				FVector RawNormal = (Corners[1] - Corners[2]) ^ (Corners[0] - Corners[2]);
-				RawNormal *= RawNormalScale;
-				double FourSquaredTriangleArea = RawNormal.SizeSquared();
-
-				// We support even small triangles, but this function is still useful to
-				// see if we have at least one valid triangle in the mesh
-				if (FourSquaredTriangleArea > 0.0f)
-				{
-					return true;
-				}
+				VertexPositions[VertexIDs[0]],
+				VertexPositions[VertexIDs[1]],
+				VertexPositions[VertexIDs[2]]
+			};
+
+			FVector RawNormal = (Corners[1] - Corners[2]) ^ (Corners[0] - Corners[2]);
+			RawNormal *= RawNormalScale;
+			double FourSquaredTriangleArea = RawNormal.SizeSquared();
+
+			// We support even small triangles, but this function is still useful to
+			// see if we have at least one valid triangle in the mesh
+			if (FourSquaredTriangleArea > 0.0f)
+			{
+				return true;
 			}
 		}
 
diff --git a/Engine/Plugins/Enterprise/StaticMeshEditorExtension/Source/MeshProcessingLibrary/Private/MeshDescriptionAdapter.cpp b/Engine/Plugins/Enterprise/StaticMeshEditorExtension/Source/MeshProcessingLibrary/Private/MeshDescriptionAdapter.cpp
--- a/Engine/Plugins/Enterprise/StaticMeshEditorExtension/Source/MeshProcessingLibrary/Private/MeshDescriptionAdapter.cpp
+++ b/Engine/Plugins/Enterprise/StaticMeshEditorExtension/Source/MeshProcessingLibrary/Private/MeshDescriptionAdapter.cpp
@@ -1496,7 +1496,7 @@ MeshSimplifier::Point3D FMeshDescriptionAdapter::GetVertexCoordinates(FMSVertexI
 
 int FMeshDescriptionAdapter::GetConnectingEdgesAtVertex(FMSVertexID VertexIndex, std::vector<FMSEdgeID>& OutEdges)
 {
-	TArrayView<const FEdgeID> VertexEdges = MeshDescription.GetVertexConnectedEdges(FVertexID(VertexIndex));
+	TArrayView<const FEdgeID> VertexEdges = MeshDescription.GetVertexConnectedEdgeIDs(FVertexID(VertexIndex));
 
 	OutEdges.resize(VertexEdges.Num());
 	for (int32 Index = 0; Index < VertexEdges.Num(); ++Index)
@@ -1514,7 +1514,7 @@ int FMeshDescriptionAdapter::GetEdgesAtVertex(FMSVertexID VertexIndex, std::vect
 		return 0;
 	}
 
-	TArrayView<const FEdgeID> VertexEdges = MeshDescription.GetVertexConnectedEdges(FVertexID(VertexIndex));
+	TArrayView<const FEdgeID> VertexEdges = MeshDescription.GetVertexConnectedEdgeIDs(FVertexID(VertexIndex));
 
 	Edges.clear();
 
@@ -1533,7 +1533,7 @@ int FMeshDescriptionAdapter::GetEdgesAtVertex(FMSVertexID VertexIndex, std::vect
 
 int FMeshDescriptionAdapter::GetVertexConnectingTriangles(FMSVertexID VertexIndex, std::unordered_set<FMSTriangleID> &OutTriangles)
 {
-	TArrayView<const FEdgeID> EdgeIDs = MeshDescription.GetVertexConnectedEdges(FVertexID(VertexIndex));
+	TArrayView<const FEdgeID> EdgeIDs = MeshDescription.GetVertexConnectedEdgeIDs(FVertexID(VertexIndex));
 
 	OutTriangles.clear();
 
@@ -1563,7 +1563,7 @@ int FMeshDescriptionAdapter::ValidateVertexTopology(FMSVertexID VertexIndex)
 
 	// Logic took from GPNode::evaluateConnectivity
 	const FVertexID VertexID(VertexIndex);
-	TArrayView<const FEdgeID> VertexConnectedEdgeIDs = MeshDescription.GetVertexConnectedEdges(VertexID);
+	TArrayView<const FEdgeID> VertexConnectedEdgeIDs = MeshDescription.GetVertexConnectedEdgeIDs(VertexID);
 
 	switch (VertexConnectedEdgeIDs.Num())
 	{
@@ -1913,7 +1913,7 @@ int FMeshDescriptionAdapter::CreateTriangle(FMSVertexID VertexIndexA, FMSVertexI
 	for (int32 Corner = 0; Corner < 3; ++Corner)
 	{
 		// vertex instance's color, uvs are copied from first instance of vertex if such instance exists
-		TArrayView<const FVertexInstanceID> VertexInstances = MeshDescription.GetVertexVertexInstances(VertexIDs[Corner]);
+		TArrayView<const FVertexInstanceID> VertexInstances = MeshDescription.GetVertexVertexInstanceIDs(VertexIDs[Corner]);
 
 		FVertexInstanceID VertexInstanceID = INDEX_NONE;
 		for (const FVertexInstanceID CandidateVertexInstance : VertexInstances)
@@ -2002,7 +2002,7 @@ int FMeshDescriptionAdapter::DeleteVertex(FMSVertexID VertexIndex)
 
 	const FVertexID VertexID(VertexIndex);
 
-	TArray<FEdgeID> EdgeIDsToDelete(MeshDescription.GetVertexConnectedEdges(VertexID));
+	TArray<FEdgeID> EdgeIDsToDelete(MeshDescription.GetVertexConnectedEdgeIDs(VertexID));
 
 	for (const FEdgeID EdgeIDToDelete : EdgeIDsToDelete)
 	{
@@ -2104,7 +2104,7 @@ int FMeshDescriptionAdapter::MergeVertices(FMSVertexID VertexIndex, FMSVertexID
 
 	FVertexID VertexID(VertexIndex);
 
-	TArray<FEdgeID> OtherVertexEdges(MeshDescription.GetVertexConnectedEdges(FVertexID(OtherVertexIndex)));
+	TArray<FEdgeID> OtherVertexEdges(MeshDescription.GetVertexConnectedEdgeIDs(FVertexID(OtherVertexIndex)));
 
 	TArray<FEdgeID> EdgeIDsToDelete;
 	for (FEdgeID EdgeID : OtherVertexEdges)
diff --git a/Engine/Plugins/Enterprise/StaticMeshEditorExtension/Source/StaticMeshEditorExtension/Private/MeshEditingContext.cpp b/Engine/Plugins/Enterprise/StaticMeshEditorExtension/Source/StaticMeshEditorExtension/Private/MeshEditingContext.cpp
--- a/Engine/Plugins/Enterprise/StaticMeshEditorExtension/Source/StaticMeshEditorExtension/Private/MeshEditingContext.cpp
+++ b/Engine/Plugins/Enterprise/StaticMeshEditorExtension/Source/StaticMeshEditorExtension/Private/MeshEditingContext.cpp
@@ -118,7 +118,7 @@ namespace MeshEditingContext
 				FVector Tangent = FVector::ZeroVector;
 				FVector Binormal = FVector::ZeroVector;
 
-				for (const FTriangleID TriangleID : MeshDescription->GetPolygonTriangleIDs(PolygonID))
+				for (const FTriangleID TriangleID : MeshDescription->GetPolygonTriangles(PolygonID))
 				{
 					TArrayView<const FVertexInstanceID> TriVertexInstanceIDs = MeshDescription->GetTriangleVertexInstances(TriangleID);
 					const FVertexID VertexID0 = MeshDescription->GetVertexInstanceVertex(TriVertexInstanceIDs[0]);
diff --git a/Engine/Plugins/Enterprise/StaticMeshEditorExtension/Source/StaticMeshEditorExtension/Private/MeshEditorUtils.cpp b/Engine/Plugins/Enterprise/StaticMeshEditorExtension/Source/StaticMeshEditorExtension/Private/MeshEditorUtils.cpp
--- a/Engine/Plugins/Enterprise/StaticMeshEditorExtension/Source/StaticMeshEditorExtension/Private/MeshEditorUtils.cpp
+++ b/Engine/Plugins/Enterprise/StaticMeshEditorExtension/Source/StaticMeshEditorExtension/Private/MeshEditorUtils.cpp
@@ -209,7 +209,7 @@ TMap<FPolygonID, FTriangleData> MeshEditingUtilsImpl::GetFilteredTriangleData(co
 	FVector TriangleVertexPositions[3];
 	for (const FPolygonID PolygonID : Polygons)
 	{
-		TArrayView<const FTriangleID> PolygonTriangleIDs = MeshDescription->GetPolygonTriangleIDs(PolygonID);
+		TArrayView<const FTriangleID> PolygonTriangleIDs = MeshDescription->GetPolygonTriangles(PolygonID);
 
 		for (const FTriangleID MeshTriangleID : PolygonTriangleIDs)
 		{
diff --git a/Engine/Plugins/Enterprise/StaticMeshEditorExtension/Source/StaticMeshEditorExtension/Private/UVTools/UVGenerationFlattenMapping.cpp b/Engine/Plugins/Enterprise/StaticMeshEditorExtension/Source/StaticMeshEditorExtension/Private/UVTools/UVGenerationFlattenMapping.cpp
--- a/Engine/Plugins/Enterprise/StaticMeshEditorExtension/Source/StaticMeshEditorExtension/Private/UVTools/UVGenerationFlattenMapping.cpp
+++ b/Engine/Plugins/Enterprise/StaticMeshEditorExtension/Source/StaticMeshEditorExtension/Private/UVTools/UVGenerationFlattenMapping.cpp
@@ -1122,7 +1122,7 @@ TArray<FUVGenerationFlattenMappingInternal::FaceStruct> FUVGenerationFlattenMapp
 
 	for (const FPolygonID PolygonID : InMesh.Polygons().GetElementIDs())
 	{
-		TArrayView<const FTriangleID> TriangleIDs = InMesh.GetPolygonTriangleIDs(PolygonID);
+		TArrayView<const FTriangleID> TriangleIDs = InMesh.GetPolygonTriangles(PolygonID);
 		for (const FTriangleID TriangleID : TriangleIDs)
 		{
 			FaceStruct Face;
@@ -1317,7 +1317,7 @@ TArray<int32> UUVGenerationFlattenMapping::GetOverlappingCornersRemapping(const
 	int32 CurrentWedge = 0;
 	for (const FPolygonID PolygonID : InMeshDescription.Polygons().GetElementIDs())
 	{
-		for (const FTriangleID& TriangleIDs : InMeshDescription.GetPolygonTriangleIDs(PolygonID))
+		for (const FTriangleID& TriangleIDs : InMeshDescription.GetPolygonTriangles(PolygonID))
 		{
 			for (int32 Corner = 0; Corner < 3; ++Corner)
 			{
diff --git a/Engine/Plugins/Experimental/BlastPlugin/Source/BlastAuthoring/Private/MeshUtility.cpp b/Engine/Plugins/Experimental/BlastPlugin/Source/BlastAuthoring/Private/MeshUtility.cpp
--- a/Engine/Plugins/Experimental/BlastPlugin/Source/BlastAuthoring/Private/MeshUtility.cpp
+++ b/Engine/Plugins/Experimental/BlastPlugin/Source/BlastAuthoring/Private/MeshUtility.cpp
@@ -53,7 +53,7 @@ void FMeshUtility::EditableMeshToBlastMesh(const UEditableMesh* SourceMesh, Nv::
 
 	for(const FPolygonID PolygonID : MeshDescription->Polygons().GetElementIDs())
 	{
-		TArrayView<const FTriangleID> TriangleIDs = MeshDescription->GetPolygonTriangleIDs(PolygonID);
+		TArrayView<const FTriangleID> TriangleIDs = MeshDescription->GetPolygonTriangles(PolygonID);
 		for(const FTriangleID TriangleID : TriangleIDs)
 		{
 			for(int32 TriVertIndex = 0; TriVertIndex < 3; ++TriVertIndex)
@@ -120,7 +120,7 @@ void FMeshUtility::EditableMeshToBlastMesh(const UEditableMesh* SourceMesh, int3
 		BlastUVs.Emplace(UV.X, UV.Y);
 	}
 
-	TArrayView<const FPolygonID> PolygonGroupIDs = MeshDescription->GetPolygonGroupPolygons(FPolygonGroupID(PolygonGroup));
+	TArrayView<const FPolygonID> PolygonGroupIDs = MeshDescription->GetPolygonGroupPolygonIDs(FPolygonGroupID(PolygonGroup));
 
 	TSharedPtr<FGeometryCollection, ESPMode::ThreadSafe> GeometryCollectionSPtr = GeometryCollection->GetGeometryCollection();
 	FGeometryCollection* GeometryCollectionPtr = GeometryCollectionSPtr.Get();
@@ -146,7 +146,7 @@ void FMeshUtility::EditableMeshToBlastMesh(const UEditableMesh* SourceMesh, int3
 
 	for(const FPolygonID PolygonID : PolygonGroupIDs)
 	{
-		TArrayView<const FTriangleID> TriangleIDs = MeshDescription->GetPolygonTriangleIDs(PolygonID);
+		TArrayView<const FTriangleID> TriangleIDs = MeshDescription->GetPolygonTriangles(PolygonID);
 
 		check(TriangleIDs.Num() == 1);
 
diff --git a/Engine/Plugins/Experimental/GeometryProcessing/Source/MeshConversion/Private/DynamicMeshToMeshDescription.cpp b/Engine/Plugins/Experimental/GeometryProcessing/Source/MeshConversion/Private/DynamicMeshToMeshDescription.cpp
--- a/Engine/Plugins/Experimental/GeometryProcessing/Source/MeshConversion/Private/DynamicMeshToMeshDescription.cpp
+++ b/Engine/Plugins/Experimental/GeometryProcessing/Source/MeshConversion/Private/DynamicMeshToMeshDescription.cpp
@@ -25,7 +25,7 @@ namespace DynamicMeshToMeshDescriptionConversionHelper
 		{
 			FPolygonGroupID PolygonGroupID = MeshOutArg.GetPolygonPolygonGroup(PolygonID);
 
-			TArrayView<const FTriangleID> TriangleIDs = MeshOutArg.GetPolygonTriangleIDs(PolygonID);
+			TArrayView<const FTriangleID> TriangleIDs = MeshOutArg.GetPolygonTriangles(PolygonID);
 			int NumTriangles = TriangleIDs.Num();
 			for (int TriIdx = 0; TriIdx < NumTriangles; ++TriIdx, ++MeshInTriIdx)
 			{
@@ -66,10 +66,11 @@ void FDynamicMeshToMeshDescription::UpdateAttributes(const FDynamicMesh3* MeshIn
 {
 	check(MeshIn->IsCompactV());
 
+	FStaticMeshAttributes Attributes(MeshOut);
+
 	if (bUpdateNormals)
 	{
-		TVertexInstanceAttributesRef<FVector> InstanceAttrib =
-			MeshOut.VertexInstanceAttributes().GetAttributesRef<FVector>(MeshAttribute::VertexInstance::Normal);
+		TVertexInstanceAttributesRef<FVector> InstanceAttrib = Attributes.GetVertexInstanceNormals();
 		ensureMsgf(InstanceAttrib.IsValid(), TEXT("Trying to update normals on a MeshDescription that has no normal attributes"));
 		if (InstanceAttrib.IsValid())
 		{
@@ -85,7 +86,7 @@ void FDynamicMeshToMeshDescription::UpdateAttributes(const FDynamicMesh3* MeshIn
 				for (int VertID : MeshIn->VertexIndicesItr())
 				{
 					FVector Normal = (FVector)MeshIn->GetVertexNormal(VertID);
-					for (FVertexInstanceID InstanceID : MeshOut.GetVertexVertexInstances(FVertexID(VertID)))
+					for (FVertexInstanceID InstanceID : MeshOut.GetVertexVertexInstanceIDs(FVertexID(VertID)))
 					{
 						InstanceAttrib.Set(InstanceID, Normal);
 					}
@@ -96,8 +97,7 @@ void FDynamicMeshToMeshDescription::UpdateAttributes(const FDynamicMesh3* MeshIn
 
 	if (bUpdateUVs)
 	{
-		TVertexInstanceAttributesRef<FVector2D> InstanceAttrib =
-			MeshOut.VertexInstanceAttributes().GetAttributesRef<FVector2D>(MeshAttribute::VertexInstance::TextureCoordinate);
+		TVertexInstanceAttributesRef<FVector2D> InstanceAttrib = Attributes.GetVertexInstanceUVs();
 		ensureMsgf(InstanceAttrib.IsValid(), TEXT("Trying to update UVs on a MeshDescription that has no texture coordinate attributes"));
 		if (InstanceAttrib.IsValid())
 		{
@@ -115,7 +115,7 @@ void FDynamicMeshToMeshDescription::UpdateAttributes(const FDynamicMesh3* MeshIn
 				for (int VertID : MeshIn->VertexIndicesItr())
 				{
 					FVector2D UV = (FVector2D)MeshIn->GetVertexUV(VertID);
-					for (FVertexInstanceID InstanceID : MeshOut.GetVertexVertexInstances(FVertexID(VertID)))
+					for (FVertexInstanceID InstanceID : MeshOut.GetVertexVertexInstanceIDs(FVertexID(VertID)))
 					{
 						InstanceAttrib.Set(InstanceID, UV);
 					}
@@ -132,11 +132,11 @@ void FDynamicMeshToMeshDescription::UpdateTangents(const FDynamicMesh3* MeshIn,
 	if (!ensureMsgf(MeshIn->IsCompactT(), TEXT("Trying to update MeshDescription Tangents from a non-compact DynamicMesh"))) return;
 	if (!ensureMsgf(MeshIn->HasAttributes(), TEXT("Trying to update MeshDescription Tangents from a DynamicMesh that has no Normals attribute"))) return;
 
+	FStaticMeshAttributes Attributes(MeshOut);
+
 	const FDynamicMeshNormalOverlay* Normals = MeshIn->Attributes()->PrimaryNormals();
-	TVertexInstanceAttributesRef<FVector> TangentAttrib =
-		MeshOut.VertexInstanceAttributes().GetAttributesRef<FVector>(MeshAttribute::VertexInstance::Tangent);
-	TVertexInstanceAttributesRef<float> BinormalSignAttrib =
-		MeshOut.VertexInstanceAttributes().GetAttributesRef<float>(MeshAttribute::VertexInstance::BinormalSign);
+	TVertexInstanceAttributesRef<FVector> TangentAttrib = Attributes.GetVertexInstanceTangents();
+	TVertexInstanceAttributesRef<float> BinormalSignAttrib = Attributes.GetVertexInstanceBinormalSigns();
 
 	if (!ensureMsgf(TangentAttrib.IsValid(), TEXT("Trying to update Tangents on a MeshDescription that has no Tangent Vertex Instance attribute"))) return;
 	if (!ensureMsgf(BinormalSignAttrib.IsValid(), TEXT("Trying to update Tangents on a MeshDescription that has no BinormalSign Vertex Instance attribute"))) return;
diff --git a/Engine/Plugins/Experimental/GeometryProcessing/Source/MeshConversion/Private/MeshDescriptionBuilder.cpp b/Engine/Plugins/Experimental/GeometryProcessing/Source/MeshConversion/Private/MeshDescriptionBuilder.cpp
--- a/Engine/Plugins/Experimental/GeometryProcessing/Source/MeshConversion/Private/MeshDescriptionBuilder.cpp
+++ b/Engine/Plugins/Experimental/GeometryProcessing/Source/MeshConversion/Private/MeshDescriptionBuilder.cpp
@@ -19,15 +19,13 @@ namespace ExtendedMeshAttribute
 
 void FMeshDescriptionBuilder::SetMeshDescription(FMeshDescription* Description)
 {
+	FStaticMeshAttributes Attributes(*Description);
+
 	this->MeshDescription = Description;
-	this->VertexPositions = 
-		MeshDescription->VertexAttributes().GetAttributesRef<FVector>(MeshAttribute::Vertex::Position);
-	this->InstanceUVs =
-		MeshDescription->VertexInstanceAttributes().GetAttributesRef<FVector2D>(MeshAttribute::VertexInstance::TextureCoordinate);
-	this->InstanceNormals =
-		MeshDescription->VertexInstanceAttributes().GetAttributesRef<FVector>(MeshAttribute::VertexInstance::Normal);
-	this->InstanceColors =
-		MeshDescription->VertexInstanceAttributes().GetAttributesRef<FVector4>(MeshAttribute::VertexInstance::Color);
+	this->VertexPositions = Attributes.GetVertexPositions();
+	this->InstanceUVs = Attributes.GetVertexInstanceUVs();
+	this->InstanceNormals = Attributes.GetVertexInstanceNormals();
+	this->InstanceColors = Attributes.GetVertexInstanceColors();
 }
 
 
diff --git a/Engine/Plugins/Experimental/GeometryProcessing/Source/MeshConversion/Private/MeshDescriptionToDynamicMesh.cpp b/Engine/Plugins/Experimental/GeometryProcessing/Source/MeshConversion/Private/MeshDescriptionToDynamicMesh.cpp
--- a/Engine/Plugins/Experimental/GeometryProcessing/Source/MeshConversion/Private/MeshDescriptionToDynamicMesh.cpp
+++ b/Engine/Plugins/Experimental/GeometryProcessing/Source/MeshConversion/Private/MeshDescriptionToDynamicMesh.cpp
@@ -127,8 +127,7 @@ void FMeshDescriptionToDynamicMesh::Convert(const FMeshDescription* MeshIn, FDyn
 {
 	// look up vertex positions
 	const FVertexArray& VertexIDs = MeshIn->Vertices();
-	TVertexAttributesConstRef<FVector> VertexPositions =
-		MeshIn->VertexAttributes().GetAttributesRef<FVector>(MeshAttribute::Vertex::Position);
+	TVertexAttributesConstRef<FVector> VertexPositions = MeshIn->GetVertexPositions();
 
 	// copy vertex positions
 	for (const FVertexID VertexID : VertexIDs.GetElementIDs())
@@ -140,10 +139,9 @@ void FMeshDescriptionToDynamicMesh::Convert(const FMeshDescription* MeshIn, FDyn
 
 	// look up vertex-instance UVs and normals
 	// @todo: does the MeshDescription always have UVs and Normals?
-	TVertexInstanceAttributesConstRef<FVector2D> InstanceUVs =
-		MeshIn->VertexInstanceAttributes().GetAttributesRef<FVector2D>(MeshAttribute::VertexInstance::TextureCoordinate);
-	TVertexInstanceAttributesConstRef<FVector> InstanceNormals =
-		MeshIn->VertexInstanceAttributes().GetAttributesRef<FVector>(MeshAttribute::VertexInstance::Normal);
+	FStaticMeshConstAttributes Attributes(*MeshIn);
+	TVertexInstanceAttributesConstRef<FVector2D> InstanceUVs = Attributes.GetVertexInstanceUVs();
+	TVertexInstanceAttributesConstRef<FVector> InstanceNormals = Attributes.GetVertexInstanceNormals();
 
 	TPolygonAttributesConstRef<int> PolyGroups =
 		MeshIn->PolygonAttributes().GetAttributesRef<int>(ExtendedMeshAttribute::PolyTriGroups);
@@ -182,7 +180,7 @@ void FMeshDescriptionToDynamicMesh::Convert(const FMeshDescription* MeshIn, FDyn
 	{
 		int32 PolygonGroupID = MeshIn->GetPolygonPolygonGroup(PolygonID).GetValue();
 
-		TArrayView<const FTriangleID> TriangleIDs = MeshIn->GetPolygonTriangleIDs(PolygonID);
+		TArrayView<const FTriangleID> TriangleIDs = MeshIn->GetPolygonTriangles(PolygonID);
 		int NumTriangles = TriangleIDs.Num();
 		for (int TriIdx = 0; TriIdx < NumTriangles; ++TriIdx)
 		{
@@ -416,16 +414,15 @@ void FMeshDescriptionToDynamicMesh::CopyTangents(const FMeshDescription* SourceM
 	if (!ensureMsgf(bCalculateMaps, TEXT("Cannot CopyTangents unless Maps were calculated"))) return;
 	if (!ensureMsgf(TriToPolyTriMap.Num() == TargetMesh->TriangleCount(), TEXT("Tried to CopyTangents to mesh with different triangle count"))) return;
 
-	TVertexInstanceAttributesConstRef<FVector> InstanceNormals =
-		SourceMesh->VertexInstanceAttributes().GetAttributesRef<FVector>(MeshAttribute::VertexInstance::Normal);
-	TVertexInstanceAttributesConstRef<FVector> InstanceTangents =
-		SourceMesh->VertexInstanceAttributes().GetAttributesRef<FVector>(MeshAttribute::VertexInstance::Tangent);
-	TVertexInstanceAttributesConstRef<float> InstanceSigns =
-		SourceMesh->VertexInstanceAttributes().GetAttributesRef<float>(MeshAttribute::VertexInstance::BinormalSign);
+	FStaticMeshConstAttributes Attributes(*SourceMesh);
 
-	if (!ensureMsgf(InstanceNormals.IsValid(), TEXT("Cannot CopyTangents from MeshDescription with invalid Instance Normals"))) return;
-	if (!ensureMsgf(InstanceTangents.IsValid(), TEXT("Cannot CopyTangents from MeshDescription with invalid Instance Tangents"))) return;
-	if (!ensureMsgf(InstanceSigns.IsValid(), TEXT("Cannot CopyTangents from MeshDescription with invalid Instance BinormalSigns"))) return;
+	TArrayView<const FVector> InstanceNormals = Attributes.GetVertexInstanceNormals().GetRawArray();
+	TArrayView<const FVector> InstanceTangents = Attributes.GetVertexInstanceTangents().GetRawArray();
+	TArrayView<const float> InstanceSigns = Attributes.GetVertexInstanceBinormalSigns().GetRawArray();
+
+	if (!ensureMsgf(InstanceNormals.Num() == 0, TEXT("Cannot CopyTangents from MeshDescription with invalid Instance Normals"))) return;
+	if (!ensureMsgf(InstanceTangents.Num() == 0, TEXT("Cannot CopyTangents from MeshDescription with invalid Instance Tangents"))) return;
+	if (!ensureMsgf(InstanceSigns.Num() == 0, TEXT("Cannot CopyTangents from MeshDescription with invalid Instance BinormalSigns"))) return;
 
 	TangentsOut->SetMesh(TargetMesh);
 	TangentsOut->InitializeTriVertexTangents(false);
@@ -433,14 +430,14 @@ void FMeshDescriptionToDynamicMesh::CopyTangents(const FMeshDescription* SourceM
 	for (int32 TriID : TargetMesh->TriangleIndicesItr())
 	{
 		FIndex2i PolyTriIdx = TriToPolyTriMap[TriID];
-		TArrayView<const FTriangleID> TriangleIDs = SourceMesh->GetPolygonTriangleIDs( FPolygonID(PolyTriIdx.A) );
+		TArrayView<const FTriangleID> TriangleIDs = SourceMesh->GetPolygonTriangles( FPolygonID(PolyTriIdx.A) );
 		const FTriangleID TriangleID = TriangleIDs[PolyTriIdx.B];
 		TArrayView<const FVertexInstanceID> InstanceTri = SourceMesh->GetTriangleVertexInstances(TriangleID);
 		for (int32 j = 0; j < 3; ++j)
 		{
-			FVector Normal = InstanceNormals.Get(InstanceTri[j], 0);
-			FVector Tangent = InstanceTangents.Get(InstanceTri[j], 0);
-			float BitangentSign = InstanceSigns.Get(InstanceTri[j], 0);
+			FVector Normal = InstanceNormals[InstanceTri[j]];
+			FVector Tangent = InstanceTangents[InstanceTri[j]];
+			float BitangentSign = InstanceSigns[InstanceTri[j]];
 			FVector3f Bitangent = VectorUtil::Bitangent((FVector3f)Normal, (FVector3f)Tangent, (float)BitangentSign);
 			Tangent.Normalize(); Bitangent.Normalize();
 			TangentsOut->SetPerTriangleTangent(TriID, j, Tangent, Bitangent);
diff --git a/Engine/Plugins/Experimental/GeometryProcessing/Source/MeshConversion/Public/MeshDescriptionAdapter.h b/Engine/Plugins/Experimental/GeometryProcessing/Source/MeshConversion/Public/MeshDescriptionAdapter.h
--- a/Engine/Plugins/Experimental/GeometryProcessing/Source/MeshConversion/Public/MeshDescriptionAdapter.h
+++ b/Engine/Plugins/Experimental/GeometryProcessing/Source/MeshConversion/Public/MeshDescriptionAdapter.h
@@ -27,8 +27,10 @@ protected:
 public:
 	FMeshDescriptionTriangleMeshAdapter(const FMeshDescription* MeshIn) : Mesh(MeshIn)
 	{
-		VertexPositions = MeshIn->VertexAttributes().GetAttributesRef<FVector>(MeshAttribute::Vertex::Position);
-		VertexInstanceNormals = MeshIn->VertexInstanceAttributes().GetAttributesRef<FVector>(MeshAttribute::VertexInstance::Normal);
+		FStaticMeshConstAttributes Attributes(*MeshIn);
+		VertexPositions = Attributes.GetVertexPositions();
+		VertexInstanceNormals = Attributes.GetVertexInstanceNormals();
+		// @todo: can we hold TArrayViews of the attribute arrays here? Do we guarantee not to mutate the mesh description for the duration of this object?
 	}
 
 	bool IsTriangle(int32 TID) const
@@ -118,8 +120,9 @@ protected:
 public:
 	FMeshDescriptionEditableTriangleMeshAdapter(FMeshDescription* MeshIn) : Mesh(MeshIn)
 	{
-		VertexPositions = MeshIn->VertexAttributes().GetAttributesRef<FVector>(MeshAttribute::Vertex::Position);
-		VertexInstanceNormals = MeshIn->VertexInstanceAttributes().GetAttributesRef<FVector>(MeshAttribute::VertexInstance::Normal);
+		FStaticMeshAttributes Attributes(*MeshIn);
+		VertexPositions = Attributes.GetVertexPositions();
+		VertexInstanceNormals = Attributes.GetVertexInstanceNormals();
 	}
 
 	bool IsTriangle(int32 TID) const
diff --git a/Engine/Plugins/Experimental/MeshModelingToolset/Source/MeshModelingTools/Private/EditPivotTool.cpp b/Engine/Plugins/Experimental/MeshModelingToolset/Source/MeshModelingTools/Private/EditPivotTool.cpp
--- a/Engine/Plugins/Experimental/MeshModelingToolset/Source/MeshModelingTools/Private/EditPivotTool.cpp
+++ b/Engine/Plugins/Experimental/MeshModelingToolset/Source/MeshModelingTools/Private/EditPivotTool.cpp
@@ -136,15 +136,13 @@ void UEditPivotTool::Shutdown(EToolShutdownType ShutdownType)
 }
 
 
-void VertexIteration(const FMeshDescription* Mesh, TFunction<void(int32, const FVector&)> ApplyFunc)
+void VertexIteration(const FMeshDescription* Mesh, TFunctionRef<void(int32, const FVector&)> ApplyFunc)
 {
-	const FVertexArray& VertexIDs = Mesh->Vertices();
-	TVertexAttributesConstRef<FVector> VertexPositions =
-		Mesh->VertexAttributes().GetAttributesRef<FVector>(MeshAttribute::Vertex::Position);
+	TArrayView<const FVector> VertexPositions = Mesh->GetVertexPositions().GetRawArray();
 
-	for (const FVertexID VertexID : VertexIDs.GetElementIDs())
+	for (const FVertexID VertexID : Mesh->Vertices().GetElementIDs())
 	{
-		const FVector Position = VertexPositions.Get(VertexID);
+		const FVector Position = VertexPositions[VertexID];
 		ApplyFunc(VertexID.GetValue(), Position);
 	}
 }
diff --git a/Engine/Plugins/Experimental/MeshModelingToolset/Source/MeshModelingToolsEditorOnly/Private/AttributeEditorTool.cpp b/Engine/Plugins/Experimental/MeshModelingToolset/Source/MeshModelingToolsEditorOnly/Private/AttributeEditorTool.cpp
--- a/Engine/Plugins/Experimental/MeshModelingToolset/Source/MeshModelingToolsEditorOnly/Private/AttributeEditorTool.cpp
+++ b/Engine/Plugins/Experimental/MeshModelingToolset/Source/MeshModelingToolsEditorOnly/Private/AttributeEditorTool.cpp
@@ -268,6 +268,8 @@ static bool RemoveAttribute(FMeshDescription* Mesh, EAttributeEditorElementType
 
 static bool IsReservedName(FName AttributeName)
 {
+	// @todo: use FAttributesSetBase::DoesAttributeHaveAnyFlags(EMeshAttributeFlags::Mandatory) to determine this
+	// will need a way to look up MeshDescription attribute sets from their category.
 	return (AttributeName == MeshAttribute::Vertex::Position)
 		|| (AttributeName == MeshAttribute::VertexInstance::TextureCoordinate)
 		|| (AttributeName == MeshAttribute::VertexInstance::Normal)
diff --git a/Engine/Plugins/Experimental/MeshModelingToolset/Source/ModelingComponents/Private/AssetUtils/MeshDescriptionUtil.cpp b/Engine/Plugins/Experimental/MeshModelingToolset/Source/ModelingComponents/Private/AssetUtils/MeshDescriptionUtil.cpp
--- a/Engine/Plugins/Experimental/MeshModelingToolset/Source/ModelingComponents/Private/AssetUtils/MeshDescriptionUtil.cpp
+++ b/Engine/Plugins/Experimental/MeshModelingToolset/Source/ModelingComponents/Private/AssetUtils/MeshDescriptionUtil.cpp
@@ -24,10 +24,6 @@ void UE::MeshDescription::InitializeAutoGeneratedAttributes(FMeshDescription& Me
 		// if neither are invalid we are not going to recompute
 		if (bHasInvalidNormals || bHasInvalidTangents)
 		{
-			// get the normal/tangent buffers
-			TVertexInstanceAttributesRef<FVector> VertexInstanceNormals = Mesh.VertexInstanceAttributes().GetAttributesRef<FVector>(MeshAttribute::VertexInstance::Normal);
-			TVertexInstanceAttributesRef<FVector> VertexInstanceTangents = Mesh.VertexInstanceAttributes().GetAttributesRef<FVector>(MeshAttribute::VertexInstance::Tangent);
-
 			FStaticMeshAttributes Attributes(Mesh);
 			if (!Attributes.GetTriangleNormals().IsValid() || !Attributes.GetTriangleTangents().IsValid())
 			{
diff --git a/Engine/Plugins/Experimental/MeshModelingToolset/Source/ModelingOperatorsEditorOnly/Private/CompositionOps/VoxelBooleanMeshesOp.cpp b/Engine/Plugins/Experimental/MeshModelingToolset/Source/ModelingOperatorsEditorOnly/Private/CompositionOps/VoxelBooleanMeshesOp.cpp
--- a/Engine/Plugins/Experimental/MeshModelingToolset/Source/ModelingOperatorsEditorOnly/Private/CompositionOps/VoxelBooleanMeshesOp.cpp
+++ b/Engine/Plugins/Experimental/MeshModelingToolset/Source/ModelingOperatorsEditorOnly/Private/CompositionOps/VoxelBooleanMeshesOp.cpp
@@ -178,10 +178,10 @@ float FVoxelBooleanMeshesOp::ComputeVoxelSize() const
 		FVector Scale = PlacedMeshes[i].Transform.GetScale3D();
 		const FMeshDescription&  MeshDescription = *PlacedMeshes[i].Mesh;
 
-		TVertexAttributesConstRef<FVector> VertexPositions = MeshDescription.VertexAttributes().GetAttributesRef<FVector>(MeshAttribute::Vertex::Position);
+		TArrayView<const FVector> VertexPositions = MeshDescription.GetVertexPositions().GetRawArray();
 		for (const FVertexID VertexID : MeshDescription.Vertices().GetElementIDs())
 		{
-			const FVector Pos = VertexPositions[VertexID];
+			const FVector& Pos = VertexPositions[VertexID];
             
 			BBoxMin.X = FMath::Min(BBoxMin.X, Pos.X);
 			BBoxMin.Y = FMath::Min(BBoxMin.Y, Pos.Y);
diff --git a/Engine/Plugins/Experimental/MeshModelingToolset/Source/ModelingOperatorsEditorOnly/Private/CompositionOps/VoxelMergeMeshesOp.cpp b/Engine/Plugins/Experimental/MeshModelingToolset/Source/ModelingOperatorsEditorOnly/Private/CompositionOps/VoxelMergeMeshesOp.cpp
--- a/Engine/Plugins/Experimental/MeshModelingToolset/Source/ModelingOperatorsEditorOnly/Private/CompositionOps/VoxelMergeMeshesOp.cpp
+++ b/Engine/Plugins/Experimental/MeshModelingToolset/Source/ModelingOperatorsEditorOnly/Private/CompositionOps/VoxelMergeMeshesOp.cpp
@@ -139,10 +139,10 @@ float FVoxelMergeMeshesOp::ComputeVoxelSize() const
 		FVector Scale = PlacedMeshes[i].Transform.GetScale3D();
 		const FMeshDescription&  MeshDescription = *PlacedMeshes[i].Mesh;
 
-		TVertexAttributesConstRef<FVector> VertexPositions = MeshDescription.VertexAttributes().GetAttributesRef<FVector>(MeshAttribute::Vertex::Position);
+		TArrayView<const FVector> VertexPositions = MeshDescription.GetVertexPositions().GetRawArray();
 		for (const FVertexID VertexID : MeshDescription.Vertices().GetElementIDs())
 		{
-			const FVector Pos = VertexPositions[VertexID];
+			const FVector& Pos = VertexPositions[VertexID];
 
 			BBoxMin.X = FMath::Min(BBoxMin.X, Pos.X);
 			BBoxMin.Y = FMath::Min(BBoxMin.Y, Pos.Y);
diff --git a/Engine/Plugins/Experimental/ProxyLODPlugin/Source/ProxyLOD/Private/ProxyLODMaterialTransferUtilities.cpp b/Engine/Plugins/Experimental/ProxyLODPlugin/Source/ProxyLOD/Private/ProxyLODMaterialTransferUtilities.cpp
--- a/Engine/Plugins/Experimental/ProxyLODPlugin/Source/ProxyLOD/Private/ProxyLODMaterialTransferUtilities.cpp
+++ b/Engine/Plugins/Experimental/ProxyLODPlugin/Source/ProxyLOD/Private/ProxyLODMaterialTransferUtilities.cpp
@@ -1254,10 +1254,12 @@ namespace
 		const ProxyLOD::FRasterGrid& DstUVGrid, const FMeshDescription& MeshDescription,
 		ProxyLOD::FLinearColorGrid& OutDstBufferGrid)
 	{
-		TVertexAttributesConstRef<FVector> VertexPositions = MeshDescription.VertexAttributes().GetAttributesRef<FVector>(MeshAttribute::Vertex::Position);
-		TVertexInstanceAttributesConstRef<FVector> VertexInstanceNormals = MeshDescription.VertexInstanceAttributes().GetAttributesRef<FVector>(MeshAttribute::VertexInstance::Normal);
-		TVertexInstanceAttributesConstRef<FVector> VertexInstanceTangents = MeshDescription.VertexInstanceAttributes().GetAttributesRef<FVector>(MeshAttribute::VertexInstance::Tangent);
-		TVertexInstanceAttributesConstRef<float> VertexInstanceBinormalSigns = MeshDescription.VertexInstanceAttributes().GetAttributesRef<float>(MeshAttribute::VertexInstance::BinormalSign);
+		FStaticMeshConstAttributes Attributes(MeshDescription);
+
+		TVertexAttributesConstRef<FVector> VertexPositions = Attributes.GetVertexPositions();
+		TVertexInstanceAttributesConstRef<FVector> VertexInstanceNormals = Attributes.GetVertexInstanceNormals();
+		TVertexInstanceAttributesConstRef<FVector> VertexInstanceTangents = Attributes.GetVertexInstanceTangents();
+		TVertexInstanceAttributesConstRef<float> VertexInstanceBinormalSigns = Attributes.GetVertexInstanceBinormalSigns();
 
 		const FIntPoint SuperSampleSize = SuperSampledNormalGrid.Size();
 		const FIntPoint DstBufferSize = OutDstBufferGrid.Size();
diff --git a/Engine/Plugins/Experimental/ProxyLODPlugin/Source/ProxyLOD/Private/ProxyLODMeshAttrTransfer.cpp b/Engine/Plugins/Experimental/ProxyLODPlugin/Source/ProxyLOD/Private/ProxyLODMeshAttrTransfer.cpp
--- a/Engine/Plugins/Experimental/ProxyLODPlugin/Source/ProxyLOD/Private/ProxyLODMeshAttrTransfer.cpp
+++ b/Engine/Plugins/Experimental/ProxyLODPlugin/Source/ProxyLOD/Private/ProxyLODMeshAttrTransfer.cpp
@@ -62,13 +62,16 @@ void ProxyLOD::TransferMeshAttributes(const FClosestPolyField& SrcPolyField, FMe
 {
 	const int32 NumFaces = InOutMesh.Polygons().Num();
 
-	TVertexAttributesRef<FVector> VertexPositions = InOutMesh.VertexAttributes().GetAttributesRef<FVector>(MeshAttribute::Vertex::Position);
-	TPolygonGroupAttributesRef<FName> PolygonGroupImportedMaterialSlotNames = InOutMesh.PolygonGroupAttributes().GetAttributesRef<FName>(MeshAttribute::PolygonGroup::ImportedMaterialSlotName);
-	TVertexInstanceAttributesRef<FVector> VertexInstanceNormals = InOutMesh.VertexInstanceAttributes().GetAttributesRef<FVector>(MeshAttribute::VertexInstance::Normal);
-	TVertexInstanceAttributesRef<FVector> VertexInstanceTangents = InOutMesh.VertexInstanceAttributes().GetAttributesRef<FVector>(MeshAttribute::VertexInstance::Tangent);
-	TVertexInstanceAttributesRef<float> VertexInstanceBinormalSigns = InOutMesh.VertexInstanceAttributes().GetAttributesRef<float>(MeshAttribute::VertexInstance::BinormalSign);
-	TVertexInstanceAttributesRef<FVector4> VertexInstanceColors = InOutMesh.VertexInstanceAttributes().GetAttributesRef<FVector4>(MeshAttribute::VertexInstance::Color);
-	TVertexInstanceAttributesRef<FVector2D> VertexInstanceUVs = InOutMesh.VertexInstanceAttributes().GetAttributesRef<FVector2D>(MeshAttribute::VertexInstance::TextureCoordinate);
+	FStaticMeshAttributes Attributes(InOutMesh);
+
+	TArrayView<const FVector> VertexPositions = Attributes.GetVertexPositions().GetRawArray();
+	TArrayView<FVector> VertexInstanceNormals = Attributes.GetVertexInstanceNormals().GetRawArray();
+	TArrayView<FVector> VertexInstanceTangents = Attributes.GetVertexInstanceTangents().GetRawArray();
+	TArrayView<float> VertexInstanceBinormalSigns = Attributes.GetVertexInstanceBinormalSigns().GetRawArray();
+	TArrayView<FVector4> VertexInstanceColors = Attributes.GetVertexInstanceColors().GetRawArray();
+
+	TPolygonGroupAttributesRef<FName> PolygonGroupImportedMaterialSlotNames = Attributes.GetPolygonGroupMaterialSlotNames();
+	TVertexInstanceAttributesRef<FVector2D> VertexInstanceUVs = Attributes.GetVertexInstanceUVs();
 
 	//const FMeshDescriptionArrayAdapter& RawMeshArrayAdapter = SrcPolyField.MeshAdapter();
 	ProxyLOD::Parallel_For(ProxyLOD::FIntRange(0, NumFaces),
@@ -81,7 +84,7 @@ void ProxyLOD::TransferMeshAttributes(const FClosestPolyField& SrcPolyField, FMe
 		{
 			FPolygonID PolygonID = FPolygonID(CurrentRange);
 			int32 LastMaterialIndex = -1;
-			for (const FTriangleID TriangleID : InOutMesh.GetPolygonTriangleIDs(PolygonID))
+			for (const FTriangleID TriangleID : InOutMesh.GetPolygonTriangles(PolygonID))
 			{
 				// get the three corners for this Triangle
 				TArrayView<const FVertexInstanceID> TriangleVerts = InOutMesh.GetTriangleVertexInstances(TriangleID);
@@ -119,6 +122,7 @@ void ProxyLOD::TransferMeshAttributes(const FClosestPolyField& SrcPolyField, FMe
 				FPolygonGroupID PolygonGroupID(LastMaterialIndex);
 				if (!InOutMesh.IsPolygonGroupValid(PolygonGroupID))
 				{
+					// @todo: check: this doesn't seem thread-safe to me - RichardTW
 					InOutMesh.CreatePolygonGroupWithID(PolygonGroupID);
 					PolygonGroupImportedMaterialSlotNames[PolygonGroupID] = FName(*FString::Printf(TEXT("ProxyLOD_Material_%d"), FMath::Rand()));
 				}
@@ -190,8 +194,10 @@ void ProxyLOD::TransferSrcNormals(const FClosestPolyField& SrcPolyField, FAOSMes
 
 void ProxyLOD::TransferVertexColors(const FClosestPolyField& SrcPolyField, FMeshDescription& InOutMesh)
 {
-	TVertexAttributesConstRef<FVector> VertexPositions = InOutMesh.VertexAttributes().GetAttributesRef<FVector>(MeshAttribute::Vertex::Position);
-	TVertexInstanceAttributesRef<FVector4> VertexInstanceColors = InOutMesh.VertexInstanceAttributes().GetAttributesRef<FVector4>(MeshAttribute::VertexInstance::Color);
+	TArrayView<const FVector> VertexPositions = InOutMesh.GetVertexPositions().GetRawArray();
+
+	FStaticMeshAttributes Attributes(InOutMesh);
+	TArrayView<FVector4> VertexInstanceColors = Attributes.GetVertexInstanceColors().GetRawArray();
 
 	const uint32 NumWedges = InOutMesh.VertexInstances().Num();
 
@@ -259,7 +265,7 @@ void ProxyLOD::TransferVertexColors(const FClosestPolyField& SrcPolyField, FMesh
 template <typename ProjectionOperatorType>
 void ProjectVerticiesOntoSrc(const ProjectionOperatorType& ProjectionOperator, const FClosestPolyField& SrcPolyField, FMeshDescription& InOutMesh)
 {
-	TVertexAttributesRef<FVector> VertexPositions = InOutMesh.VertexAttributes().GetAttributesRef<FVector>(MeshAttribute::Vertex::Position);
+	TArrayView<FVector> VertexPositions = InOutMesh.GetVertexPositions().GetRawArray();
 	const uint32 NumVertexes = InOutMesh.Vertices().Num();
 
 	ProxyLOD::Parallel_For(ProxyLOD::FUIntRange(0, NumVertexes),
diff --git a/Engine/Plugins/Experimental/ProxyLODPlugin/Source/ProxyLOD/Private/ProxyLODMeshConvertUtils.cpp b/Engine/Plugins/Experimental/ProxyLODPlugin/Source/ProxyLOD/Private/ProxyLODMeshConvertUtils.cpp
--- a/Engine/Plugins/Experimental/ProxyLODPlugin/Source/ProxyLOD/Private/ProxyLODMeshConvertUtils.cpp
+++ b/Engine/Plugins/Experimental/ProxyLODPlugin/Source/ProxyLOD/Private/ProxyLODMeshConvertUtils.cpp
@@ -20,14 +20,15 @@ void ProxyLOD::MixedPolyMeshToRawMesh(const FMixedPolyMesh& SimpleMesh, FMeshDes
 
 	ResetStaticMeshDescription(DstRawMesh);
 
-	TVertexAttributesRef<FVector> VertexPositions = DstRawMesh.VertexAttributes().GetAttributesRef<FVector>(MeshAttribute::Vertex::Position);
-	TEdgeAttributesRef<bool> EdgeHardnesses = DstRawMesh.EdgeAttributes().GetAttributesRef<bool>(MeshAttribute::Edge::IsHard);
-	TPolygonGroupAttributesRef<FName> PolygonGroupImportedMaterialSlotNames = DstRawMesh.PolygonGroupAttributes().GetAttributesRef<FName>(MeshAttribute::PolygonGroup::ImportedMaterialSlotName);
-	TVertexInstanceAttributesRef<FVector> VertexInstanceNormals = DstRawMesh.VertexInstanceAttributes().GetAttributesRef<FVector>(MeshAttribute::VertexInstance::Normal);
-	TVertexInstanceAttributesRef<FVector> VertexInstanceTangents = DstRawMesh.VertexInstanceAttributes().GetAttributesRef<FVector>(MeshAttribute::VertexInstance::Tangent);
-	TVertexInstanceAttributesRef<float> VertexInstanceBinormalSigns = DstRawMesh.VertexInstanceAttributes().GetAttributesRef<float>(MeshAttribute::VertexInstance::BinormalSign);
-	TVertexInstanceAttributesRef<FVector4> VertexInstanceColors = DstRawMesh.VertexInstanceAttributes().GetAttributesRef<FVector4>(MeshAttribute::VertexInstance::Color);
-	TVertexInstanceAttributesRef<FVector2D> VertexInstanceUVs = DstRawMesh.VertexInstanceAttributes().GetAttributesRef<FVector2D>(MeshAttribute::VertexInstance::TextureCoordinate);
+	FStaticMeshAttributes Attributes(DstRawMesh);
+	TVertexAttributesRef<FVector> VertexPositions = Attributes.GetVertexPositions();
+	TEdgeAttributesRef<bool> EdgeHardnesses = Attributes.GetEdgeHardnesses();
+	TPolygonGroupAttributesRef<FName> PolygonGroupImportedMaterialSlotNames = Attributes.GetPolygonGroupMaterialSlotNames();
+	TVertexInstanceAttributesRef<FVector> VertexInstanceNormals = Attributes.GetVertexInstanceNormals();
+	TVertexInstanceAttributesRef<FVector> VertexInstanceTangents = Attributes.GetVertexInstanceTangents();
+	TVertexInstanceAttributesRef<float> VertexInstanceBinormalSigns = Attributes.GetVertexInstanceBinormalSigns();
+	TVertexInstanceAttributesRef<FVector4> VertexInstanceColors = Attributes.GetVertexInstanceColors();
+	TVertexInstanceAttributesRef<FVector2D> VertexInstanceUVs = Attributes.GetVertexInstanceUVs();
 
 	// Splitting a quad doesn't introduce any new verts.
 	const uint32 DstNumVerts = SimpleMesh.Points.size();
@@ -134,14 +135,15 @@ void ProxyLOD::AOSMeshToRawMesh(const FAOSMesh& AOSMesh, FMeshDescription& OutRa
 
 	ResetStaticMeshDescription(OutRawMesh);
 
-	TVertexAttributesRef<FVector> VertexPositions = OutRawMesh.VertexAttributes().GetAttributesRef<FVector>(MeshAttribute::Vertex::Position);
-	TEdgeAttributesRef<bool> EdgeHardnesses = OutRawMesh.EdgeAttributes().GetAttributesRef<bool>(MeshAttribute::Edge::IsHard);
-	TPolygonGroupAttributesRef<FName> PolygonGroupImportedMaterialSlotNames = OutRawMesh.PolygonGroupAttributes().GetAttributesRef<FName>(MeshAttribute::PolygonGroup::ImportedMaterialSlotName);
-	TVertexInstanceAttributesRef<FVector> VertexInstanceNormals = OutRawMesh.VertexInstanceAttributes().GetAttributesRef<FVector>(MeshAttribute::VertexInstance::Normal);
-	TVertexInstanceAttributesRef<FVector> VertexInstanceTangents = OutRawMesh.VertexInstanceAttributes().GetAttributesRef<FVector>(MeshAttribute::VertexInstance::Tangent);
-	TVertexInstanceAttributesRef<float> VertexInstanceBinormalSigns = OutRawMesh.VertexInstanceAttributes().GetAttributesRef<float>(MeshAttribute::VertexInstance::BinormalSign);
-	TVertexInstanceAttributesRef<FVector4> VertexInstanceColors = OutRawMesh.VertexInstanceAttributes().GetAttributesRef<FVector4>(MeshAttribute::VertexInstance::Color);
-	TVertexInstanceAttributesRef<FVector2D> VertexInstanceUVs = OutRawMesh.VertexInstanceAttributes().GetAttributesRef<FVector2D>(MeshAttribute::VertexInstance::TextureCoordinate);
+	FStaticMeshAttributes Attributes(OutRawMesh);
+	TVertexAttributesRef<FVector> VertexPositions = Attributes.GetVertexPositions();
+	TEdgeAttributesRef<bool> EdgeHardnesses = Attributes.GetEdgeHardnesses();
+	TPolygonGroupAttributesRef<FName> PolygonGroupImportedMaterialSlotNames = Attributes.GetPolygonGroupMaterialSlotNames();
+	TVertexInstanceAttributesRef<FVector> VertexInstanceNormals = Attributes.GetVertexInstanceNormals();
+	TVertexInstanceAttributesRef<FVector> VertexInstanceTangents = Attributes.GetVertexInstanceTangents();
+	TVertexInstanceAttributesRef<float> VertexInstanceBinormalSigns = Attributes.GetVertexInstanceBinormalSigns();
+	TVertexInstanceAttributesRef<FVector4> VertexInstanceColors = Attributes.GetVertexInstanceColors();
+	TVertexInstanceAttributesRef<FVector2D> VertexInstanceUVs = Attributes.GetVertexInstanceUVs();
 
 	const uint32 DstNumPositions = AOSMesh.GetNumVertexes();
 	const uint32 DstNumIndexes = AOSMesh.GetNumIndexes();
@@ -224,14 +226,15 @@ void ProxyLOD::VertexDataMeshToRawMesh(const FVertexDataMesh& SrcVertexDataMesh,
 {
 	ResetStaticMeshDescription(OutRawMesh);
 
-	TVertexAttributesRef<FVector> VertexPositions = OutRawMesh.VertexAttributes().GetAttributesRef<FVector>(MeshAttribute::Vertex::Position);
-	TEdgeAttributesRef<bool> EdgeHardnesses = OutRawMesh.EdgeAttributes().GetAttributesRef<bool>(MeshAttribute::Edge::IsHard);
-	TPolygonGroupAttributesRef<FName> PolygonGroupImportedMaterialSlotNames = OutRawMesh.PolygonGroupAttributes().GetAttributesRef<FName>(MeshAttribute::PolygonGroup::ImportedMaterialSlotName);
-	TVertexInstanceAttributesRef<FVector> VertexInstanceNormals = OutRawMesh.VertexInstanceAttributes().GetAttributesRef<FVector>(MeshAttribute::VertexInstance::Normal);
-	TVertexInstanceAttributesRef<FVector> VertexInstanceTangents = OutRawMesh.VertexInstanceAttributes().GetAttributesRef<FVector>(MeshAttribute::VertexInstance::Tangent);
-	TVertexInstanceAttributesRef<float> VertexInstanceBinormalSigns = OutRawMesh.VertexInstanceAttributes().GetAttributesRef<float>(MeshAttribute::VertexInstance::BinormalSign);
-	TVertexInstanceAttributesRef<FVector4> VertexInstanceColors = OutRawMesh.VertexInstanceAttributes().GetAttributesRef<FVector4>(MeshAttribute::VertexInstance::Color);
-	TVertexInstanceAttributesRef<FVector2D> VertexInstanceUVs = OutRawMesh.VertexInstanceAttributes().GetAttributesRef<FVector2D>(MeshAttribute::VertexInstance::TextureCoordinate);
+	FStaticMeshAttributes Attributes(OutRawMesh);
+	TVertexAttributesRef<FVector> VertexPositions = Attributes.GetVertexPositions();
+	TEdgeAttributesRef<bool> EdgeHardnesses = Attributes.GetEdgeHardnesses();
+	TPolygonGroupAttributesRef<FName> PolygonGroupImportedMaterialSlotNames = Attributes.GetPolygonGroupMaterialSlotNames();
+	TVertexInstanceAttributesRef<FVector> VertexInstanceNormals = Attributes.GetVertexInstanceNormals();
+	TVertexInstanceAttributesRef<FVector> VertexInstanceTangents = Attributes.GetVertexInstanceTangents();
+	TVertexInstanceAttributesRef<float> VertexInstanceBinormalSigns = Attributes.GetVertexInstanceBinormalSigns();
+	TVertexInstanceAttributesRef<FVector4> VertexInstanceColors = Attributes.GetVertexInstanceColors();
+	TVertexInstanceAttributesRef<FVector2D> VertexInstanceUVs = Attributes.GetVertexInstanceUVs();
 
 	const uint32 DstNumPositions = SrcVertexDataMesh.Points.Num();
 	const uint32 DstNumIndexes = SrcVertexDataMesh.Indices.Num();
@@ -364,11 +367,12 @@ void ProxyLOD::VertexDataMeshToRawMesh(const FVertexDataMesh& SrcVertexDataMesh,
 
 void ProxyLOD::RawMeshToVertexDataMesh(const FMeshDescription& SrcRawMesh, FVertexDataMesh& DstVertexDataMesh)
 {
-	TVertexAttributesConstRef<FVector> VertexPositions = SrcRawMesh.VertexAttributes().GetAttributesRef<FVector>(MeshAttribute::Vertex::Position);
-	TVertexInstanceAttributesConstRef<FVector> VertexInstanceNormals = SrcRawMesh.VertexInstanceAttributes().GetAttributesRef<FVector>(MeshAttribute::VertexInstance::Normal);
-	TVertexInstanceAttributesConstRef<FVector> VertexInstanceTangents = SrcRawMesh.VertexInstanceAttributes().GetAttributesRef<FVector>(MeshAttribute::VertexInstance::Tangent);
-	TVertexInstanceAttributesConstRef<float> VertexInstanceBinormalSigns = SrcRawMesh.VertexInstanceAttributes().GetAttributesRef<float>(MeshAttribute::VertexInstance::BinormalSign);
-	TVertexInstanceAttributesConstRef<FVector2D> VertexInstanceUVs = SrcRawMesh.VertexInstanceAttributes().GetAttributesRef<FVector2D>(MeshAttribute::VertexInstance::TextureCoordinate);
+	FStaticMeshConstAttributes Attributes(SrcRawMesh);
+	TVertexAttributesConstRef<FVector> VertexPositions = Attributes.GetVertexPositions();
+	TVertexInstanceAttributesConstRef<FVector> VertexInstanceNormals = Attributes.GetVertexInstanceNormals();
+	TVertexInstanceAttributesConstRef<FVector> VertexInstanceTangents = Attributes.GetVertexInstanceTangents();
+	TVertexInstanceAttributesConstRef<float> VertexInstanceBinormalSigns = Attributes.GetVertexInstanceBinormalSigns();
+	TVertexInstanceAttributesConstRef<FVector2D> VertexInstanceUVs = Attributes.GetVertexInstanceUVs();
 
 
 	const uint32 DstNumPositions = SrcRawMesh.Vertices().Num();
@@ -408,36 +412,33 @@ void ProxyLOD::RawMeshToVertexDataMesh(const FMeshDescription& SrcRawMesh, FVert
 	ResizeArray(DstUVs, DstNumPositions);
 
 	//Iterate all triangle and add the indices
-	for (const FPolygonID PolygonID : SrcRawMesh.Polygons().GetElementIDs())
+	for (const FTriangleID TriangleID : SrcRawMesh.Triangles().GetElementIDs())
 	{
-		for (const FTriangleID TriangleID : SrcRawMesh.GetPolygonTriangleIDs(PolygonID))
+		for (int32 Corner = 0; Corner < 3; ++Corner)
 		{
-			for (int32 Corner = 0; Corner < 3; ++Corner)
+			const FVertexInstanceID& VertexInstanceID = SrcRawMesh.GetTriangleVertexInstance(TriangleID, Corner);
+			DstIndices[VertexInstanceCount] = VertexIDToDstVertexIndex[SrcRawMesh.GetVertexInstanceVertex(VertexInstanceID)];
+
+			// Copy the tangent space:
+			// NB: The tangent space is stored per-index in the raw mesh, but only per-vertex in the vertex data mesh.
+			// We assume that the raw mesh per-index data is really duplicated per-vertex data!
+			DstTangentArray[DstIndices[VertexInstanceCount]] = VertexInstanceTangents[VertexInstanceID];
+			DstBiTangentArray[DstIndices[VertexInstanceCount]] = FVector::CrossProduct(VertexInstanceNormals[VertexInstanceID], VertexInstanceTangents[VertexInstanceID]).GetSafeNormal() * VertexInstanceBinormalSigns[VertexInstanceID];
+			DstNormalArray[DstIndices[VertexInstanceCount]] = VertexInstanceNormals[VertexInstanceID];
+
+			// Copy the UVs:
+			// NB: The UVs is stored per-index in the raw mesh, but only per-vertex in the vertex data mesh.
+			// We assume that the raw mesh per-index data is really duplicated per-vertex data!
+			if (VertexInstanceUVs.GetNumChannels() == 0)
+			{
+				DstUVs[DstIndices[VertexInstanceCount]] = FVector2D(0.0f, 0.0f);
+			}
+			else
 			{
-				const FVertexInstanceID& VertexInstanceID = SrcRawMesh.GetTriangleVertexInstance(TriangleID, Corner);
-				DstIndices[VertexInstanceCount] = VertexIDToDstVertexIndex[SrcRawMesh.GetVertexInstanceVertex(VertexInstanceID)];
-
-				// Copy the tangent space:
-				// NB: The tangent space is stored per-index in the raw mesh, but only per-vertex in the vertex data mesh.
-				// We assume that the raw mesh per-index data is really duplicated per-vertex data!
-				DstTangentArray[DstIndices[VertexInstanceCount]] = VertexInstanceTangents[VertexInstanceID];
-				DstBiTangentArray[DstIndices[VertexInstanceCount]] = FVector::CrossProduct(VertexInstanceNormals[VertexInstanceID], VertexInstanceTangents[VertexInstanceID]).GetSafeNormal() * VertexInstanceBinormalSigns[VertexInstanceID];
-				DstNormalArray[DstIndices[VertexInstanceCount]] = VertexInstanceNormals[VertexInstanceID];
-
-				// Copy the UVs:
-				// NB: The UVs is stored per-index in the raw mesh, but only per-vertex in the vertex data mesh.
-				// We assume that the raw mesh per-index data is really duplicated per-vertex data!
-				if (VertexInstanceUVs.GetNumChannels() == 0)
-				{
-					DstUVs[DstIndices[VertexInstanceCount]] = FVector2D(0.0f, 0.0f);
-				}
-				else
-				{
-					DstUVs[DstIndices[VertexInstanceCount]] = VertexInstanceUVs.Get(VertexInstanceID, 0);
-				}
-
-				VertexInstanceCount++;
+				DstUVs[DstIndices[VertexInstanceCount]] = VertexInstanceUVs.Get(VertexInstanceID, 0);
 			}
+
+			VertexInstanceCount++;
 		}
 	}
 
diff --git a/Engine/Plugins/Experimental/ProxyLODPlugin/Source/ProxyLOD/Private/ProxyLODMeshParameterization.cpp b/Engine/Plugins/Experimental/ProxyLODPlugin/Source/ProxyLOD/Private/ProxyLODMeshParameterization.cpp
--- a/Engine/Plugins/Experimental/ProxyLODPlugin/Source/ProxyLOD/Private/ProxyLODMeshParameterization.cpp
+++ b/Engine/Plugins/Experimental/ProxyLODPlugin/Source/ProxyLOD/Private/ProxyLODMeshParameterization.cpp
@@ -476,7 +476,9 @@ void ProxyLOD::GenerateAdjacency(const FMeshDescription& RawMesh, std::vector<ui
 															 // Allocate adjacency
 	AdjacencyArray.resize(AdjacencySize);
 
-	TVertexAttributesConstRef<FVector> VertexPositionsAttribute = RawMesh.VertexAttributes().GetAttributesRef<FVector>(MeshAttribute::Vertex::Position);
+	// @todo: possibility to pass the vertex position raw array directly into the below function.
+	// Can it be assumed that it is not sparse? 
+	TArrayView<const FVector> VertexPositionsAttribute = RawMesh.GetVertexPositions().GetRawArray();
 	TArray<FVector> VertexPositions;
 	VertexPositions.AddZeroed(NumVerts);
 	for (const FVertexID VertexID : RawMesh.Vertices().GetElementIDs())
diff --git a/Engine/Plugins/Experimental/ProxyLODPlugin/Source/ProxyLOD/Private/ProxyLODMeshTypes.cpp b/Engine/Plugins/Experimental/ProxyLODPlugin/Source/ProxyLOD/Private/ProxyLODMeshTypes.cpp
--- a/Engine/Plugins/Experimental/ProxyLODPlugin/Source/ProxyLOD/Private/ProxyLODMeshTypes.cpp
+++ b/Engine/Plugins/Experimental/ProxyLODPlugin/Source/ProxyLOD/Private/ProxyLODMeshTypes.cpp
@@ -19,20 +19,16 @@ FMeshDescriptionAdapter::FMeshDescriptionAdapter(const FMeshDescriptionAdapter&
 
 void FMeshDescriptionAdapter::InitializeCacheData()
 {
-	VertexPositions = RawMesh->VertexAttributes().GetAttributesRef<FVector>(MeshAttribute::Vertex::Position);
+	VertexPositions = RawMesh->GetVertexPositions();
 	TriangleCount = RawMesh->Triangles().Num();
 
 	IndexBuffer.Reserve(TriangleCount * 3);
 
-	for (const FPolygonID PolygonID : RawMesh->Polygons().GetElementIDs())
+	for (const FTriangleID TriangleID : RawMesh->Triangles().GetElementIDs())
 	{
-		TArrayView<const FTriangleID> TriangleIDs = RawMesh->GetPolygonTriangleIDs(PolygonID);
-		for (const FTriangleID TriangleID : TriangleIDs)
-		{
-			IndexBuffer.Add(RawMesh->GetTriangleVertexInstance(TriangleID, 0));
-			IndexBuffer.Add(RawMesh->GetTriangleVertexInstance(TriangleID, 1));
-			IndexBuffer.Add(RawMesh->GetTriangleVertexInstance(TriangleID, 2));
-		}
+		IndexBuffer.Add(RawMesh->GetTriangleVertexInstance(TriangleID, 0));
+		IndexBuffer.Add(RawMesh->GetTriangleVertexInstance(TriangleID, 1));
+		IndexBuffer.Add(RawMesh->GetTriangleVertexInstance(TriangleID, 2));
 	}
 
 
@@ -79,15 +75,11 @@ void FMeshDescriptionArrayAdapter::Construct(int32 MeshCount, TFunctionRef<const
 		IndexBufferArray.push_back(std::vector<FVertexInstanceID>());
 		std::vector<FVertexInstanceID>& IndexBuffer = IndexBufferArray[MeshIdx];
 		IndexBuffer.reserve(MeshPolyCount * 3);
-		for (const FPolygonID PolygonID : RawMesh->Polygons().GetElementIDs())
+		for (const FTriangleID TriangleID : RawMesh->Triangles().GetElementIDs())
 		{
-			TArrayView<const FTriangleID> TriangleIDs = RawMesh->GetPolygonTriangleIDs(PolygonID);
-			for (const FTriangleID TriangleID : TriangleIDs)
-			{
-				IndexBuffer.push_back(RawMesh->GetTriangleVertexInstance(TriangleID, 0));
-				IndexBuffer.push_back(RawMesh->GetTriangleVertexInstance(TriangleID, 1));
-				IndexBuffer.push_back(RawMesh->GetTriangleVertexInstance(TriangleID, 2));
-			}
+			IndexBuffer.push_back(RawMesh->GetTriangleVertexInstance(TriangleID, 0));
+			IndexBuffer.push_back(RawMesh->GetTriangleVertexInstance(TriangleID, 1));
+			IndexBuffer.push_back(RawMesh->GetTriangleVertexInstance(TriangleID, 2));
 		}
 
 		PolyCount += MeshPolyCount;
@@ -247,7 +239,7 @@ FMeshDescriptionArrayAdapter::FRawPoly FMeshDescriptionArrayAdapter::GetRawPoly(
 	RawPoly.FaceMaterialIndex = RawMesh.GetPolygonPolygonGroup(PolygonID).GetValue();
 	RawPoly.FaceSmoothingMask = AttributesGetter->FaceSmoothingMasks[LocalFaceNumber];
 
-	for (const FTriangleID TriangleID : RawMesh.GetPolygonTriangleIDs(PolygonID))
+	for (const FTriangleID TriangleID : RawMesh.GetPolygonTriangles(PolygonID))
 	{
 		TArrayView<const FVertexInstanceID> VertexInstanceIDs = RawMesh.GetTriangleVertexInstances(TriangleID);
 
diff --git a/Engine/Plugins/Experimental/ProxyLODPlugin/Source/ProxyLOD/Private/ProxyLODMeshTypes.h b/Engine/Plugins/Experimental/ProxyLODPlugin/Source/ProxyLOD/Private/ProxyLODMeshTypes.h
--- a/Engine/Plugins/Experimental/ProxyLODPlugin/Source/ProxyLOD/Private/ProxyLODMeshTypes.h
+++ b/Engine/Plugins/Experimental/ProxyLODPlugin/Source/ProxyLOD/Private/ProxyLODMeshTypes.h
@@ -325,12 +325,14 @@ public:
 	public:
 		FMeshDescriptionAttributesGetter(const FMeshDescription* MeshDescription)
 		{
-			VertexPositions = MeshDescription->VertexAttributes().GetAttributesRef<FVector>(MeshAttribute::Vertex::Position);
-			VertexInstanceNormals = MeshDescription->VertexInstanceAttributes().GetAttributesRef<FVector>(MeshAttribute::VertexInstance::Normal);
-			VertexInstanceTangents = MeshDescription->VertexInstanceAttributes().GetAttributesRef<FVector>(MeshAttribute::VertexInstance::Tangent);
-			VertexInstanceBinormalSigns = MeshDescription->VertexInstanceAttributes().GetAttributesRef<float>(MeshAttribute::VertexInstance::BinormalSign);
-			VertexInstanceColors = MeshDescription->VertexInstanceAttributes().GetAttributesRef<FVector4>(MeshAttribute::VertexInstance::Color);
-			VertexInstanceUVs = MeshDescription->VertexInstanceAttributes().GetAttributesRef<FVector2D>(MeshAttribute::VertexInstance::TextureCoordinate);
+			FStaticMeshConstAttributes Attributes(*MeshDescription);
+
+			VertexPositions = Attributes.GetVertexPositions();
+			VertexInstanceNormals = Attributes.GetVertexInstanceNormals();
+			VertexInstanceTangents = Attributes.GetVertexInstanceTangents();
+			VertexInstanceBinormalSigns = Attributes.GetVertexInstanceBinormalSigns();
+			VertexInstanceColors = Attributes.GetVertexInstanceColors();
+			VertexInstanceUVs = Attributes.GetVertexInstanceUVs();
 
 			TriangleCount = MeshDescription->Triangles().Num();
 			FaceSmoothingMasks.AddZeroed(TriangleCount);
diff --git a/Engine/Plugins/Experimental/ProxyLODPlugin/Source/ProxyLOD/Private/ProxyLODMeshUtilities.cpp b/Engine/Plugins/Experimental/ProxyLODPlugin/Source/ProxyLOD/Private/ProxyLODMeshUtilities.cpp
--- a/Engine/Plugins/Experimental/ProxyLODPlugin/Source/ProxyLOD/Private/ProxyLODMeshUtilities.cpp
+++ b/Engine/Plugins/Experimental/ProxyLODPlugin/Source/ProxyLOD/Private/ProxyLODMeshUtilities.cpp
@@ -1445,7 +1445,7 @@ int32 CorrectCollapsedWalls(const ProxyLOD::FkDOPTree& kDOPTree,
 {
 	typedef uint32 EdgeIdType;
 	typedef uint32 FaceIdType;
-	TVertexAttributesRef<FVector> VertexPositions = MeshDescription.VertexAttributes().GetAttributesRef<FVector>(MeshAttribute::Vertex::Position);
+	TVertexAttributesRef<FVector> VertexPositions = MeshDescription.GetVertexPositions();
 
 	ProxyLOD::FUnitTransformDataProvider kDOPDataProvider(kDOPTree);
 
@@ -1646,22 +1646,19 @@ void ProxyLOD::ColorPartitions(FMeshDescription& InOutRawMesh, const std::vector
 		FColor(153, 102, 0), FColor(249, 129, 162), FColor(29, 143, 177), FColor(118, 42, 145),
 		FColor(255, 121, 75), FColor(102, 204, 51), FColor(153, 153, 255), FColor(255, 255, 255) };
 
-	TVertexInstanceAttributesRef<FVector4> VertexInstanceColors = InOutRawMesh.VertexInstanceAttributes().GetAttributesRef<FVector4>(MeshAttribute::VertexInstance::Color);
+	TArrayView<FVector4> VertexInstanceColors = FStaticMeshAttributes(InOutRawMesh).GetVertexInstanceColors().GetRawArray();
 
 	//Remap the vertex instance
 	int32 TriangleIndex = 0;
 	TMap<uint32, FVertexInstanceID> WedgeIndexToVertexInstanceID;
 	WedgeIndexToVertexInstanceID.Reserve(InOutRawMesh.VertexInstances().Num());
-	for (const FPolygonID PolygonID : InOutRawMesh.Polygons().GetElementIDs())
+	for (const FTriangleID TriangleID : InOutRawMesh.Triangles().GetElementIDs())
 	{
-		for (const FTriangleID TriangleID : InOutRawMesh.GetPolygonTriangleIDs(PolygonID))
+		for (int32 Corner = 0; Corner < 3; ++Corner)
 		{
-			for (int32 Corner = 0; Corner < 3; ++Corner)
-			{
-				WedgeIndexToVertexInstanceID.Add((TriangleIndex * 3) + Corner, InOutRawMesh.GetTriangleVertexInstance(TriangleID, Corner));
-			}
-			TriangleIndex++;
+			WedgeIndexToVertexInstanceID.Add((TriangleIndex * 3) + Corner, InOutRawMesh.GetTriangleVertexInstance(TriangleID, Corner));
 		}
+		TriangleIndex++;
 	}
 
 	for (int i = 0; i < partitionResults.size(); ++i)
@@ -1704,20 +1701,17 @@ void ProxyLOD::AddWedgeColors(FMeshDescription& RawMesh)
 		FColor(153, 102, 0), FColor(249, 129, 162), FColor(29, 143, 177), FColor(118, 42, 145),
 		FColor(255, 121, 75), FColor(102, 204, 51), FColor(153, 153, 255), FColor(255, 255, 255) };
 
-	TVertexInstanceAttributesRef<FVector4> VertexInstanceColors = RawMesh.VertexInstanceAttributes().GetAttributesRef<FVector4>(MeshAttribute::VertexInstance::Color);
+	TArrayView<FVector4> VertexInstanceColors = FStaticMeshAttributes(RawMesh).GetVertexInstanceColors().GetRawArray();
 
 	//Recolor the vertex instances
 	int32 TriangleIndex = 0;
-	for (const FPolygonID PolygonID : RawMesh.Polygons().GetElementIDs())
+	for (const FTriangleID TriangleID : RawMesh.Triangles().GetElementIDs())
 	{
-		for (const FTriangleID TriangleID : RawMesh.GetPolygonTriangleIDs(PolygonID))
+		for (int32 Corner = 0; Corner < 3; ++Corner)
 		{
-			for (int32 Corner = 0; Corner < 3; ++Corner)
-			{
-				VertexInstanceColors[RawMesh.GetTriangleVertexInstance(TriangleID, Corner)] = FLinearColor(ColorRange[((TriangleIndex*3) + Corner) % 13]);
-			}
-			TriangleIndex++;
+			VertexInstanceColors[RawMesh.GetTriangleVertexInstance(TriangleID, Corner)] = FLinearColor(ColorRange[((TriangleIndex*3) + Corner) % 13]);
 		}
+		TriangleIndex++;
 	}
 }
 
diff --git a/Engine/Plugins/Experimental/ProxyLODPlugin/Source/ProxyLOD/Private/ProxyLODRasterizer.cpp b/Engine/Plugins/Experimental/ProxyLODPlugin/Source/ProxyLOD/Private/ProxyLODRasterizer.cpp
--- a/Engine/Plugins/Experimental/ProxyLODPlugin/Source/ProxyLOD/Private/ProxyLODRasterizer.cpp
+++ b/Engine/Plugins/Experimental/ProxyLODPlugin/Source/ProxyLOD/Private/ProxyLODRasterizer.cpp
@@ -220,7 +220,7 @@ namespace ProxyLOD
 		FIntPoint Size = OutMaterial.GetPropertySize(EFlattenMaterialProperties::Diffuse);
 		ResizeArray(ColorBuffer, Size.X * Size.Y);
 
-		TVertexInstanceAttributesConstRef<FVector4> VertexInstanceColors = RawMesh.VertexInstanceAttributes().GetAttributesRef<FVector4>(MeshAttribute::VertexInstance::Color);
+		TArrayView<const FVector4> VertexInstanceColors = FStaticMeshConstAttributes(RawMesh).GetVertexInstanceColors().GetRawArray();
 
 		for (int j = 0; j < Size.Y; ++j)
 		{
diff --git a/Engine/Plugins/Experimental/ProxyLODPlugin/Source/ProxyLOD/Private/ProxyLODkDOPUtils.cpp b/Engine/Plugins/Experimental/ProxyLODPlugin/Source/ProxyLOD/Private/ProxyLODkDOPUtils.cpp
--- a/Engine/Plugins/Experimental/ProxyLODPlugin/Source/ProxyLOD/Private/ProxyLODkDOPUtils.cpp
+++ b/Engine/Plugins/Experimental/ProxyLODPlugin/Source/ProxyLOD/Private/ProxyLODkDOPUtils.cpp
@@ -40,7 +40,7 @@ void ProxyLOD::BuildkDOPTree(const FMeshDescription& MeshDescription, FkDOPTree&
 {
 	TRACE_CPUPROFILER_EVENT_SCOPE(ProxyLOD::BuildkDOPTree)
 
-	TVertexAttributesConstRef<FVector> VertexPositions = MeshDescription.VertexAttributes().GetAttributesRef<FVector>(MeshAttribute::Vertex::Position);
+	TArrayView<const FVector> VertexPositions = MeshDescription.GetVertexPositions().GetRawArray();
 
 	uint32 NumSrcPoly = MeshDescription.Triangles().Num();
 
diff --git a/Engine/Plugins/Runtime/EditableMesh/Source/EditableMesh/EditableGeometryCollectionAdapter.cpp b/Engine/Plugins/Runtime/EditableMesh/Source/EditableMesh/EditableGeometryCollectionAdapter.cpp
--- a/Engine/Plugins/Runtime/EditableMesh/Source/EditableMesh/EditableGeometryCollectionAdapter.cpp
+++ b/Engine/Plugins/Runtime/EditableMesh/Source/EditableMesh/EditableGeometryCollectionAdapter.cpp
@@ -324,7 +324,7 @@ void UEditableGeometryCollectionAdapter::InitializeFromEditableMesh( const UEdit
 		FAdaptorPolygon& RenderingPolygon = RenderingPolygons[ PolygonID ];
 		RenderingPolygon.PolygonGroupID = PolygonGroupID;
 
-		TArrayView<const FTriangleID> TriangleIDs = MeshDescription->GetPolygonTriangleIDs( PolygonID );
+		TArrayView<const FTriangleID> TriangleIDs = MeshDescription->GetPolygonTriangles( PolygonID );
 		for( const FTriangleID TriangleID : TriangleIDs )
 		{
 			FAdaptorPolygon2Group::FMeshTriangle Triangle;
@@ -406,7 +406,7 @@ void UEditableGeometryCollectionAdapter::OnRebuildRenderMesh(const UEditableMesh
 				const FPolygonGroupID& PolygonGroupID = MeshDescription->GetPolygonPolygonGroup(PolygonID);
 				int32 PolygonGroupIDValue = PolygonGroupID.GetValue();
 				int32 PolygonIDValue = PolygonID.GetValue();
-				TArrayView<const FTriangleID> TriangleIDs = MeshDescription->GetPolygonTriangleIDs(PolygonID);
+				TArrayView<const FTriangleID> TriangleIDs = MeshDescription->GetPolygonTriangles(PolygonID);
 				for (const FTriangleID TriangleID : TriangleIDs)
 				{
 					uint32 ElementIndex = Collection->AddElements(1, FGeometryCollection::FacesGroup);
@@ -658,7 +658,7 @@ void UEditableGeometryCollectionAdapter::OnSetVertexAttribute(const UEditableMes
 					const FVertexInstanceArray& VertexInstances = EditableMesh->GetMeshDescription()->VertexInstances();
 
 					// Set the vertex buffer position of all of the vertex instances for this editable vertex
-					for (const FVertexInstanceID VertexInstanceID : MeshDescription->GetVertexVertexInstances(VertexID))
+					for (const FVertexInstanceID VertexInstanceID : MeshDescription->GetVertexVertexInstanceIDs(VertexID))
 					{
 						check(MeshDescription->IsVertexInstanceValid(VertexInstanceID));
 						auto& OldPosition = GCVertices[VertexInstanceID.GetValue()]; // temp for debug
diff --git a/Engine/Plugins/Runtime/EditableMesh/Source/EditableMesh/EditableMesh.cpp b/Engine/Plugins/Runtime/EditableMesh/Source/EditableMesh/EditableMesh.cpp
--- a/Engine/Plugins/Runtime/EditableMesh/Source/EditableMesh/EditableMesh.cpp
+++ b/Engine/Plugins/Runtime/EditableMesh/Source/EditableMesh/EditableMesh.cpp
@@ -880,13 +880,13 @@ bool UEditableMesh::IsOrphanedVertex( const FVertexID VertexID ) const
 
 int32 UEditableMesh::GetVertexConnectedEdgeCount( const FVertexID VertexID ) const
 {
-	return GetMeshDescription()->GetVertexConnectedEdges( VertexID ).Num();
+	return GetMeshDescription()->GetVertexConnectedEdgeIDs( VertexID ).Num();
 }
 
 
 FEdgeID UEditableMesh::GetVertexConnectedEdge( const FVertexID VertexID, const int32 ConnectedEdgeNumber ) const
 {
-	TArrayView<const FEdgeID> ConnectedEdgeIDs = GetMeshDescription()->GetVertexConnectedEdges( VertexID );
+	TArrayView<const FEdgeID> ConnectedEdgeIDs = GetMeshDescription()->GetVertexConnectedEdgeIDs( VertexID );
 	return ( ConnectedEdgeNumber >= 0 && ConnectedEdgeNumber < ConnectedEdgeIDs.Num() ) ? ConnectedEdgeIDs[ ConnectedEdgeNumber ] : INDEX_NONE;
 }
 
@@ -962,13 +962,13 @@ bool UEditableMesh::IsValidPolygonGroup( const FPolygonGroupID PolygonGroupID )
 
 int32 UEditableMesh::GetPolygonCountInGroup( const FPolygonGroupID PolygonGroupID ) const
 {
-	return GetMeshDescription()->GetPolygonGroupPolygons( PolygonGroupID ).Num();
+	return GetMeshDescription()->GetPolygonGroupPolygonIDs( PolygonGroupID ).Num();
 }
 
 
 FPolygonID UEditableMesh::GetPolygonInGroup( const FPolygonGroupID PolygonGroupID, const int32 PolygonNumber ) const
 {
-	TArrayView<const FPolygonID> PolygonIDs = GetMeshDescription()->GetPolygonGroupPolygons( PolygonGroupID );
+	TArrayView<const FPolygonID> PolygonIDs = GetMeshDescription()->GetPolygonGroupPolygonIDs( PolygonGroupID );
 	return ( PolygonNumber >= 0 && PolygonNumber < PolygonIDs.Num() ) ? PolygonIDs[ PolygonNumber ] : INDEX_NONE;
 }
 
@@ -1013,13 +1013,13 @@ FVertexID UEditableMesh::GetPolygonPerimeterVertex( const FPolygonID PolygonID,
 
 int32 UEditableMesh::GetPolygonTriangulatedTriangleCount( const FPolygonID PolygonID ) const
 {
-	return GetMeshDescription()->GetPolygonTriangleIDs( PolygonID ).Num();
+	return GetMeshDescription()->GetPolygonTriangles( PolygonID ).Num();
 }
 
 
 FTriangleID UEditableMesh::GetPolygonTriangulatedTriangle( const FPolygonID PolygonID, int32 PolygonTriangleNumber ) const
 {
-	TArrayView<const FTriangleID> Triangles = GetMeshDescription()->GetPolygonTriangleIDs( PolygonID );
+	TArrayView<const FTriangleID> Triangles = GetMeshDescription()->GetPolygonTriangles( PolygonID );
 	return ( PolygonTriangleNumber >= 0 && PolygonTriangleNumber < Triangles.Num() ) ? Triangles[ PolygonTriangleNumber ] : INDEX_NONE;
 }
 
@@ -1260,7 +1260,7 @@ void UEditableMesh::SplitVerticesIfNecessary( const TArray<FVertexID>& VerticesT
 		// Look at each vertex instance in turn.
 		// Take a copy because splitting them will mutate the list we are iterating.
 		static TArray<FVertexInstanceID> VertexInstanceIDs;
-		VertexInstanceIDs = GetMeshDescription()->GetVertexVertexInstances( VertexToSplit );
+		VertexInstanceIDs = GetMeshDescription()->GetVertexVertexInstanceIDs( VertexToSplit );
 
 		for( const FVertexInstanceID VertexInstanceID : VertexInstanceIDs )
 		{
@@ -1328,7 +1328,7 @@ void UEditableMesh::MergeVertexInstances()
 		GetConnectedSoftEdges( VertexID, VertexConnectedSoftEdges );
 
 		// Get all vertex instances of this vertex...
-		TArrayView<const FVertexInstanceID> VertexInstanceIDs = GetMeshDescription()->GetVertexVertexInstances( VertexID );
+		TArrayView<const FVertexInstanceID> VertexInstanceIDs = GetMeshDescription()->GetVertexVertexInstanceIDs( VertexID );
 
 		// ...and iterate through pairs of vertex instances, looking for potential to merge them
 		for( int32 IndexA = 0; IndexA < VertexInstanceIDs.Num() - 1; ++IndexA )
@@ -1419,7 +1419,7 @@ bool UEditableMesh::IsPreviewingSubdivisions() const
 
 void UEditableMesh::GetVertexConnectedEdges( const FVertexID VertexID, TArray<FEdgeID>& OutConnectedEdgeIDs ) const
 {
-	OutConnectedEdgeIDs = GetMeshDescription()->GetVertexConnectedEdges( VertexID );
+	OutConnectedEdgeIDs = GetMeshDescription()->GetVertexConnectedEdgeIDs( VertexID );
 }
 
 
@@ -1502,7 +1502,7 @@ void UEditableMesh::GetEdgeLoopElements( const FEdgeID EdgeID, TArray<FEdgeID>&
 			FEdgeID AdjacentEdgeID = INDEX_NONE;
 
 			// Iterate through all edges connected to this vertex
-			for( const FEdgeID ConnectedEdgeID : GetMeshDescription()->GetVertexConnectedEdges( ConnectedVertexID ) )
+			for( const FEdgeID ConnectedEdgeID : GetMeshDescription()->GetVertexConnectedEdgeIDs( ConnectedVertexID ) )
 			{
 				// If this edge hasn't been added to the loop...
 				if( !EdgeIDs.Contains( ConnectedEdgeID ) )
@@ -1976,7 +1976,7 @@ void UEditableMesh::GenerateOpenSubdivLimitSurfaceData()
 				{
 					FirstPolygonNumberForPolygonGroups.Add( NumPolygonsSoFar );
 
-					TArrayView<const FPolygonID> PolygonIDs = GetMeshDescription()->GetPolygonGroupPolygons( PolygonGroupID );
+					TArrayView<const FPolygonID> PolygonIDs = GetMeshDescription()->GetPolygonGroupPolygonIDs( PolygonGroupID );
 					NumPolygonsSoFar += PolygonIDs.Num();
 
 					for( const FPolygonID PolygonID : PolygonIDs )
@@ -2374,7 +2374,7 @@ FTriangleID UEditableMesh::ComputeBarycentricWeightForPointOnPolygon( const FPol
 	TVertexAttributesConstRef<FVector> VertexPositions = GetMeshDescription()->VertexAttributes().GetAttributesRef<FVector>( MeshAttribute::Vertex::Position );
 
 	// Figure out which triangle the incoming point is within
-	for( const FTriangleID TriangleID : GetMeshDescription()->GetPolygonTriangleIDs( PolygonID ) )
+	for( const FTriangleID TriangleID : GetMeshDescription()->GetPolygonTriangles( PolygonID ) )
 	{
 		const FVector TriangleVertex0Position = VertexPositions[ GetMeshDescription()->GetVertexInstanceVertex( GetMeshDescription()->GetTriangleVertexInstance( TriangleID, 0 ) ) ];
 		const FVector TriangleVertex1Position = VertexPositions[ GetMeshDescription()->GetVertexInstanceVertex( GetMeshDescription()->GetTriangleVertexInstance( TriangleID, 1 ) ) ];
@@ -3452,8 +3452,8 @@ void UEditableMesh::DeleteOrphanVertices( const TArray<FVertexID>& VertexIDsToDe
 
 			// Make sure the vertex is truly an orphan.  We're not going to be able to restore its polygon vertex attributes,
 			// because the polygons won't exist when we're restoring the change
-			check( GetMeshDescription()->GetVertexConnectedEdges( VertexID ).Num() == 0 );
-			check( GetMeshDescription()->GetVertexVertexInstances( VertexID ).Num() == 0 );
+			check( GetMeshDescription()->GetVertexConnectedEdgeIDs( VertexID ).Num() == 0 );
+			check( GetMeshDescription()->GetVertexVertexInstanceIDs( VertexID ).Num() == 0 );
 
 			RevertInput.VerticesToCreate.Emplace();
 			FVertexToCreate& VertexToCreate = RevertInput.VerticesToCreate.Last();
@@ -4454,7 +4454,7 @@ void UEditableMesh::ChangePolygonsVertexInstances( const TArray<FChangeVertexIns
 
 FVertexInstanceID UEditableMesh::GetVertexInstanceInPolygonForVertex( const FPolygonID PolygonID, const FVertexID VertexID ) const
 {
-	for( const FVertexInstanceID VertexInstanceID : GetMeshDescription()->GetVertexVertexInstances( VertexID ) )
+	for( const FVertexInstanceID VertexInstanceID : GetMeshDescription()->GetVertexVertexInstanceIDs( VertexID ) )
 	{
 		if( GetMeshDescription()->GetVertexInstanceConnectedPolygons( VertexInstanceID ).Contains( PolygonID ) )
 		{
@@ -4471,7 +4471,7 @@ void UEditableMesh::GetConnectedSoftEdges( const FVertexID VertexID, TArray<FEdg
 	OutConnectedSoftEdges.Reset();
 
 	TEdgeAttributesConstRef<bool> EdgeHardnesses = GetMeshDescription()->EdgeAttributes().GetAttributesRef<bool>( MeshAttribute::Edge::IsHard );
-	for( const FEdgeID ConnectedEdgeID : GetMeshDescription()->GetVertexConnectedEdges( VertexID ) )
+	for( const FEdgeID ConnectedEdgeID : GetMeshDescription()->GetVertexConnectedEdgeIDs( VertexID ) )
 	{
 		if( !EdgeHardnesses[ ConnectedEdgeID ] )
 		{
@@ -4872,7 +4872,7 @@ void UEditableMesh::TryToRemoveVertex( const FVertexID VertexID, bool& bOutWasVe
 		//// Delete the vertex instances and subsequently orphaned vertex
 		//{
 		//	// Take a copy of the array, because it will be modified by the DeleteVertexInstances call
-		//	TArray<FVertexInstanceID> VertexInstanceIDs = GetMeshDescription()->GetVertexVertexInstances( VertexID );
+		//	TArray<FVertexInstanceID> VertexInstanceIDs = GetMeshDescription()->GetVertexVertexInstanceIDs( VertexID );
 
 		//	const bool bDeleteOrphanedVertices = true;
 		//	DeleteVertexInstances( VertexInstanceIDs, bDeleteOrphanedVertices );
@@ -5913,7 +5913,7 @@ void UEditableMesh::GeneratePolygonTangentsAndNormals( const TArray<FPolygonID>&
 		FVector Tangent(ForceInitToZero);
 		FVector Binormal(ForceInitToZero);
 
-		for( const FTriangleID TriangleID : MD->GetPolygonTriangleIDs( PolygonID ) )
+		for( const FTriangleID TriangleID : MD->GetPolygonTriangles( PolygonID ) )
 		{
 			const FVertexInstanceID VertexInstanceID0 = MD->GetTriangleVertexInstance( TriangleID, 0 );
 			const FVertexInstanceID VertexInstanceID1 = MD->GetTriangleVertexInstance( TriangleID, 1 );
@@ -6756,7 +6756,7 @@ void UEditableMesh::TriangulatePolygons( const TArray<FPolygonID>& PolygonIDs, T
 
 			// Build polygons for each of the triangles that made up the original
 			{
-				for( const FTriangleID& TriangleID : GetMeshDescription()->GetPolygonTriangleIDs( PolygonID ) )
+				for( const FTriangleID& TriangleID : GetMeshDescription()->GetPolygonTriangles( PolygonID ) )
 				{
 					PolygonsToCreate.Emplace();
 					FPolygonToCreate& PolygonToCreate = PolygonsToCreate.Last();
@@ -6836,7 +6836,7 @@ void UEditableMesh::AssignPolygonsToPolygonGroups( const TArray<FPolygonGroupFor
 		GetMeshDescription()->SetPolygonPolygonGroup( PolygonID, NewPolygonGroupID );
 
 		// If old group is now empty (and we're deleting orphans), add it to the list to delete
-		if( bDeleteOrphanedPolygonGroups && GetMeshDescription()->GetPolygonGroupPolygons( OldPolygonGroupID ).Num() == 0 )
+		if( bDeleteOrphanedPolygonGroups && GetMeshDescription()->GetPolygonGroupPolygonIDs( OldPolygonGroupID ).Num() == 0 )
 		{
 			PolygonGroupsToDelete.Add( OldPolygonGroupID );
 		}
diff --git a/Engine/Plugins/Runtime/EditableMesh/Source/EditableMesh/EditableStaticMeshAdapter.cpp b/Engine/Plugins/Runtime/EditableMesh/Source/EditableMesh/EditableStaticMeshAdapter.cpp
--- a/Engine/Plugins/Runtime/EditableMesh/Source/EditableMesh/EditableStaticMeshAdapter.cpp
+++ b/Engine/Plugins/Runtime/EditableMesh/Source/EditableMesh/EditableStaticMeshAdapter.cpp
@@ -422,7 +422,7 @@ void UEditableStaticMeshAdapter::InitializeFromEditableMesh( const UEditableMesh
 		FRenderingPolygon& RenderingPolygon = RenderingPolygons[ PolygonID ];
 		RenderingPolygon.PolygonGroupID = PolygonGroupID;
 
-		TArrayView<const FTriangleID> TriangleIDs = MeshDescription->GetPolygonTriangleIDs( PolygonID );
+		TArrayView<const FTriangleID> TriangleIDs = MeshDescription->GetPolygonTriangles( PolygonID );
 		for( const FTriangleID TriangleID : TriangleIDs )
 		{
 			FRenderingPolygonGroup::FMeshTriangle Triangle;
@@ -832,7 +832,7 @@ void UEditableStaticMeshAdapter::OnReindexElements( const UEditableMesh* Editabl
 		RenderingPolygonGroup.MaxTriangles = RenderingPolygonGroup.Triangles.GetArraySize();
 
 		// Fix up references in referencing polygons
-		for( const FPolygonID PolygonID : EditableMesh->GetMeshDescription()->GetPolygonGroupPolygons( PolygonGroupID ) )
+		for( const FPolygonID PolygonID : EditableMesh->GetMeshDescription()->GetPolygonGroupPolygonIDs( PolygonGroupID ) )
 		{
 			FRenderingPolygon& RenderingPolygon = RenderingPolygons[ PolygonID ];
 			RenderingPolygon.PolygonGroupID = PolygonGroupID;
@@ -1076,7 +1076,7 @@ void UEditableStaticMeshAdapter::OnSetVertexAttribute( const UEditableMesh* Edit
 			const FVertexInstanceArray& VertexInstances = EditableMesh->GetMeshDescription()->VertexInstances();
 
 			// Set the vertex buffer position of all of the vertex instances for this editable vertex
-			for( const FVertexInstanceID VertexInstanceID : MeshDescription->GetVertexVertexInstances( VertexID ) )
+			for( const FVertexInstanceID VertexInstanceID : MeshDescription->GetVertexVertexInstanceIDs( VertexID ) )
 			{
 				check( MeshDescription->IsVertexInstanceValid( VertexInstanceID ) );
 				StaticMeshLOD.VertexBuffers.PositionVertexBuffer.VertexPosition( VertexInstanceID.GetValue() ) = NewVertexPosition;
@@ -1296,7 +1296,7 @@ void UEditableStaticMeshAdapter::OnRetriangulatePolygons( const UEditableMesh* E
 		FRenderingPolygon& RenderingPolygon = RenderingPolygons[ PolygonID ];
 		const FPolygonGroupID PolygonGroupID = RenderingPolygon.PolygonGroupID;
 		FRenderingPolygonGroup& RenderingPolygonGroup = RenderingPolygonGroups[ PolygonGroupID ];
-		TArrayView<const FTriangleID> TriangleIDs = MeshDescription->GetPolygonTriangleIDs( PolygonID );
+		TArrayView<const FTriangleID> TriangleIDs = MeshDescription->GetPolygonTriangles( PolygonID );
 
 		// Check to see whether the index buffer needs to be updated
 		bool bNeedsUpdatedTriangles = false;
diff --git a/Engine/Plugins/Runtime/EditableMesh/Source/EditableMesh/GeometryHitTest.cpp b/Engine/Plugins/Runtime/EditableMesh/Source/EditableMesh/GeometryHitTest.cpp
--- a/Engine/Plugins/Runtime/EditableMesh/Source/EditableMesh/GeometryHitTest.cpp
+++ b/Engine/Plugins/Runtime/EditableMesh/Source/EditableMesh/GeometryHitTest.cpp
@@ -92,7 +92,7 @@ FEditableMeshElementAddress FGeometryTests::QueryElement(const UEditableMesh& Ed
 	// Check polygons first; this is so we always impose a closest hit location at the poly before checking other elements, so anything behind is occluded
 	for (const FPolygonID PolygonID : FrontFacingPolygons)
 	{
-		for (const FTriangleID TriangleID : MeshDescription->GetPolygonTriangleIDs(PolygonID))
+		for (const FTriangleID TriangleID : MeshDescription->GetPolygonTriangles(PolygonID))
 		{
 			FVector TriangleVertexPositions[3];
 			for (uint32 TriangleVertexNumber = 0; TriangleVertexNumber < 3; ++TriangleVertexNumber)
diff --git a/Engine/Plugins/Tests/EditorTests/Source/EditorTests/Private/UnrealEd/MeshDescriptionTests.cpp b/Engine/Plugins/Tests/EditorTests/Source/EditorTests/Private/UnrealEd/MeshDescriptionTests.cpp
--- a/Engine/Plugins/Tests/EditorTests/Source/EditorTests/Private/UnrealEd/MeshDescriptionTests.cpp
+++ b/Engine/Plugins/Tests/EditorTests/Source/EditorTests/Private/UnrealEd/MeshDescriptionTests.cpp
@@ -356,26 +356,28 @@ bool FMeshDescriptionTest::CompareMeshDescription(const FString& AssetName, FAut
 {
 	//////////////////////////////////////////////////////////////////////////
 	//Gather the reference data
-	TVertexAttributesConstRef<FVector> ReferenceVertexPositions = ReferenceMeshDescription.VertexAttributes().GetAttributesRef<FVector>(MeshAttribute::Vertex::Position);
-	TVertexInstanceAttributesConstRef<FVector> ReferenceVertexInstanceNormals = ReferenceMeshDescription.VertexInstanceAttributes().GetAttributesRef<FVector>(MeshAttribute::VertexInstance::Normal);
-	TVertexInstanceAttributesConstRef<FVector> ReferenceVertexInstanceTangents = ReferenceMeshDescription.VertexInstanceAttributes().GetAttributesRef<FVector>(MeshAttribute::VertexInstance::Tangent);
-	TVertexInstanceAttributesConstRef<float> ReferenceVertexInstanceBinormalSigns = ReferenceMeshDescription.VertexInstanceAttributes().GetAttributesRef<float>(MeshAttribute::VertexInstance::BinormalSign);
-	TVertexInstanceAttributesConstRef<FVector4> ReferenceVertexInstanceColors = ReferenceMeshDescription.VertexInstanceAttributes().GetAttributesRef<FVector4>(MeshAttribute::VertexInstance::Color);
-	TVertexInstanceAttributesConstRef<FVector2D> ReferenceVertexInstanceUVs = ReferenceMeshDescription.VertexInstanceAttributes().GetAttributesRef<FVector2D>(MeshAttribute::VertexInstance::TextureCoordinate);
-	TEdgeAttributesConstRef<bool> ReferenceEdgeHardnesses = ReferenceMeshDescription.EdgeAttributes().GetAttributesRef<bool>(MeshAttribute::Edge::IsHard);
-	TPolygonGroupAttributesConstRef<FName> ReferencePolygonGroupMaterialName = ReferenceMeshDescription.PolygonGroupAttributes().GetAttributesRef<FName>(MeshAttribute::PolygonGroup::ImportedMaterialSlotName);
+	FStaticMeshConstAttributes ReferenceAttributes(ReferenceMeshDescription);
+	TVertexAttributesConstRef<FVector> ReferenceVertexPositions = ReferenceAttributes.GetVertexPositions();
+	TVertexInstanceAttributesConstRef<FVector> ReferenceVertexInstanceNormals = ReferenceAttributes.GetVertexInstanceNormals();
+	TVertexInstanceAttributesConstRef<FVector> ReferenceVertexInstanceTangents = ReferenceAttributes.GetVertexInstanceTangents();
+	TVertexInstanceAttributesConstRef<float> ReferenceVertexInstanceBinormalSigns = ReferenceAttributes.GetVertexInstanceBinormalSigns();
+	TVertexInstanceAttributesConstRef<FVector4> ReferenceVertexInstanceColors = ReferenceAttributes.GetVertexInstanceColors();
+	TVertexInstanceAttributesConstRef<FVector2D> ReferenceVertexInstanceUVs = ReferenceAttributes.GetVertexInstanceUVs();
+	TEdgeAttributesConstRef<bool> ReferenceEdgeHardnesses = ReferenceAttributes.GetEdgeHardnesses();
+	TPolygonGroupAttributesConstRef<FName> ReferencePolygonGroupMaterialName = ReferenceAttributes.GetPolygonGroupMaterialSlotNames();
 
 
 	//////////////////////////////////////////////////////////////////////////
 	//Gather the result data
-	TVertexAttributesConstRef<FVector> ResultVertexPositions = MeshDescription.VertexAttributes().GetAttributesRef<FVector>(MeshAttribute::Vertex::Position);
-	TVertexInstanceAttributesConstRef<FVector> ResultVertexInstanceNormals = MeshDescription.VertexInstanceAttributes().GetAttributesRef<FVector>(MeshAttribute::VertexInstance::Normal);
-	TVertexInstanceAttributesConstRef<FVector> ResultVertexInstanceTangents = MeshDescription.VertexInstanceAttributes().GetAttributesRef<FVector>(MeshAttribute::VertexInstance::Tangent);
-	TVertexInstanceAttributesConstRef<float> ResultVertexInstanceBinormalSigns = MeshDescription.VertexInstanceAttributes().GetAttributesRef<float>(MeshAttribute::VertexInstance::BinormalSign);
-	TVertexInstanceAttributesConstRef<FVector4> ResultVertexInstanceColors = MeshDescription.VertexInstanceAttributes().GetAttributesRef<FVector4>(MeshAttribute::VertexInstance::Color);
-	TVertexInstanceAttributesConstRef<FVector2D> ResultVertexInstanceUVs = MeshDescription.VertexInstanceAttributes().GetAttributesRef<FVector2D>(MeshAttribute::VertexInstance::TextureCoordinate);
-	TEdgeAttributesConstRef<bool> ResultEdgeHardnesses = MeshDescription.EdgeAttributes().GetAttributesRef<bool>(MeshAttribute::Edge::IsHard);
-	TPolygonGroupAttributesConstRef<FName> ResultPolygonGroupMaterialName = MeshDescription.PolygonGroupAttributes().GetAttributesRef<FName>(MeshAttribute::PolygonGroup::ImportedMaterialSlotName);
+	FStaticMeshConstAttributes ResultAttributes(MeshDescription);
+	TVertexAttributesConstRef<FVector> ResultVertexPositions = ResultAttributes.GetVertexPositions();
+	TVertexInstanceAttributesConstRef<FVector> ResultVertexInstanceNormals = ResultAttributes.GetVertexInstanceNormals();
+	TVertexInstanceAttributesConstRef<FVector> ResultVertexInstanceTangents = ResultAttributes.GetVertexInstanceTangents();
+	TVertexInstanceAttributesConstRef<float> ResultVertexInstanceBinormalSigns = ResultAttributes.GetVertexInstanceBinormalSigns();
+	TVertexInstanceAttributesConstRef<FVector4> ResultVertexInstanceColors = ResultAttributes.GetVertexInstanceColors();
+	TVertexInstanceAttributesConstRef<FVector2D> ResultVertexInstanceUVs = ResultAttributes.GetVertexInstanceUVs();
+	TEdgeAttributesConstRef<bool> ResultEdgeHardnesses = ResultAttributes.GetEdgeHardnesses();
+	TPolygonGroupAttributesConstRef<FName> ResultPolygonGroupMaterialName = ResultAttributes.GetPolygonGroupMaterialSlotNames();
 
 	
 	//////////////////////////////////////////////////////////////////////////
@@ -566,13 +568,14 @@ bool FMeshDescriptionTest::NTBTest(FAutomationTestExecutionInfo& ExecutionInfo)
 		FRawMesh RawMesh;
 		AssetMesh->GetSourceModel(0).LoadRawMesh(RawMesh);
 
-		//const TVertexAttributeArray<FVector>& VertexPositions = MeshDescription.VertexAttributes().GetAttributes<FVector>(MeshAttribute::Vertex::Position);
-		const TVertexInstanceAttributesRef<FVector> VertexInstanceNormals = MeshDescription.VertexInstanceAttributes().GetAttributesRef<FVector>(MeshAttribute::VertexInstance::Normal);
-		const TVertexInstanceAttributesRef<FVector> VertexInstanceTangents = MeshDescription.VertexInstanceAttributes().GetAttributesRef<FVector>(MeshAttribute::VertexInstance::Tangent);
-		const TVertexInstanceAttributesRef<float> VertexInstanceBinormalSigns = MeshDescription.VertexInstanceAttributes().GetAttributesRef<float>(MeshAttribute::VertexInstance::BinormalSign);
-		const TVertexInstanceAttributesRef<FVector4> VertexInstanceColors = MeshDescription.VertexInstanceAttributes().GetAttributesRef<FVector4>(MeshAttribute::VertexInstance::Color);
-		const TVertexInstanceAttributesRef<FVector2D> VertexInstanceUVs = MeshDescription.VertexInstanceAttributes().GetAttributesRef<FVector2D>(MeshAttribute::VertexInstance::TextureCoordinate);
+		FStaticMeshAttributes Attributes(MeshDescription);
+		const TVertexInstanceAttributesRef<FVector> VertexInstanceNormals = Attributes.GetVertexInstanceNormals();
+		const TVertexInstanceAttributesRef<FVector> VertexInstanceTangents = Attributes.GetVertexInstanceTangents();
+		const TVertexInstanceAttributesRef<float> VertexInstanceBinormalSigns = Attributes.GetVertexInstanceBinormalSigns();
+		const TVertexInstanceAttributesRef<FVector4> VertexInstanceColors = Attributes.GetVertexInstanceColors();
+		const TVertexInstanceAttributesRef<FVector2D> VertexInstanceUVs = Attributes.GetVertexInstanceUVs();
 		int32 ExistingUVCount = VertexInstanceUVs.GetNumChannels();
+
 		//Build the normals and tangent and compare the result
 		FStaticMeshOperations::ComputeTriangleTangentsAndNormals(MeshDescription, SMALL_NUMBER);
 		EComputeNTBsFlags ComputeNTBsFlags = EComputeNTBsFlags::Normals | EComputeNTBsFlags::Tangents;
@@ -613,7 +616,7 @@ bool FMeshDescriptionTest::NTBTest(FAutomationTestExecutionInfo& ExecutionInfo)
 			}
 			const FPolygonGroupID& PolygonGroupID = MeshDescription.GetPolygonPolygonGroup(PolygonID);
 			int32 PolygonIDValue = PolygonID.GetValue();
-			TArrayView<const FTriangleID> TriangleIDs = MeshDescription.GetPolygonTriangleIDs(PolygonID);
+			TArrayView<const FTriangleID> TriangleIDs = MeshDescription.GetPolygonTriangles(PolygonID);
 			for (const FTriangleID TriangleID : TriangleIDs)
 			{
 				if (bError)
@@ -695,8 +698,7 @@ IMPLEMENT_SIMPLE_AUTOMATION_TEST(FMeshDescriptionBuilderTest, "Editor.Meshes.Mes
 bool FMeshDescriptionBuilderTest::RunTest(const FString& Parameters)
 {
 	FMeshDescription MeshDescription;
-	MeshDescription.VertexAttributes().RegisterAttribute<FVector>(MeshAttribute::Vertex::Position);
-	TVertexAttributesRef<FVector> Positions = MeshDescription.VertexAttributes().GetAttributesRef<FVector>(MeshAttribute::Vertex::Position);
+	TVertexAttributesRef<FVector> Positions = MeshDescription.GetVertexPositions();
 
 	// Build a hexagonal cylinder
 
@@ -758,11 +760,11 @@ bool FMeshDescriptionBuilderTest::RunTest(const FString& Parameters)
 	for (int32 Index = 0; Index < NumSides; ++Index)
 	{
 		TestEqual("Vertex instances per vertex", MeshDescription.GetNumVertexVertexInstances(TopVertexIDs[Index]), 2);
-		TestTrue("Vertex contains vertex instance", MeshDescription.GetVertexVertexInstances(TopVertexIDs[Index]).Contains(TopVertexInstanceIDs[Index]));
-		TestTrue("Vertex contains vertex instance", MeshDescription.GetVertexVertexInstances(TopVertexIDs[Index]).Contains(SideTopVertexInstanceIDs[Index]));
+		TestTrue("Vertex contains vertex instance", MeshDescription.GetVertexVertexInstanceIDs(TopVertexIDs[Index]).Contains(TopVertexInstanceIDs[Index]));
+		TestTrue("Vertex contains vertex instance", MeshDescription.GetVertexVertexInstanceIDs(TopVertexIDs[Index]).Contains(SideTopVertexInstanceIDs[Index]));
 		TestEqual("Vertex instances per vertex", MeshDescription.GetNumVertexVertexInstances(BottomVertexIDs[Index]), 2);
-		TestTrue("Vertex contains vertex instance", MeshDescription.GetVertexVertexInstances(BottomVertexIDs[Index]).Contains(BottomVertexInstanceIDs[Index]));
-		TestTrue("Vertex contains vertex instance", MeshDescription.GetVertexVertexInstances(BottomVertexIDs[Index]).Contains(SideBottomVertexInstanceIDs[Index]));
+		TestTrue("Vertex contains vertex instance", MeshDescription.GetVertexVertexInstanceIDs(BottomVertexIDs[Index]).Contains(BottomVertexInstanceIDs[Index]));
+		TestTrue("Vertex contains vertex instance", MeshDescription.GetVertexVertexInstanceIDs(BottomVertexIDs[Index]).Contains(SideBottomVertexInstanceIDs[Index]));
 		TestEqual("Vertex instance parent", MeshDescription.GetVertexInstanceVertex(TopVertexInstanceIDs[Index]), TopVertexIDs[Index]);
 		TestEqual("Vertex instance parent", MeshDescription.GetVertexInstanceVertex(SideTopVertexInstanceIDs[Index]), TopVertexIDs[Index]);
 		TestEqual("Vertex instance parent", MeshDescription.GetVertexInstanceVertex(BottomVertexInstanceIDs[Index]), BottomVertexIDs[Index]);
@@ -775,8 +777,8 @@ bool FMeshDescriptionBuilderTest::RunTest(const FString& Parameters)
 	{
 		SideEdgeIDs[Index] = MeshDescription.CreateEdge(BottomVertexIDs[Index], TopVertexIDs[Index]);
 
-		TestTrue("Vertex connected edges", MeshDescription.GetVertexConnectedEdges(TopVertexIDs[Index]).Contains(SideEdgeIDs[Index]));
-		TestTrue("Vertex connected edges", MeshDescription.GetVertexConnectedEdges(BottomVertexIDs[Index]).Contains(SideEdgeIDs[Index]));
+		TestTrue("Vertex connected edges", MeshDescription.GetVertexConnectedEdgeIDs(TopVertexIDs[Index]).Contains(SideEdgeIDs[Index]));
+		TestTrue("Vertex connected edges", MeshDescription.GetVertexConnectedEdgeIDs(BottomVertexIDs[Index]).Contains(SideEdgeIDs[Index]));
 	}
 
 	// Add polygon groups
diff --git a/Engine/Source/Developer/Datasmith/DatasmithExporter/Private/DatasmithMeshExporter.cpp b/Engine/Source/Developer/Datasmith/DatasmithExporter/Private/DatasmithMeshExporter.cpp
--- a/Engine/Source/Developer/Datasmith/DatasmithExporter/Private/DatasmithMeshExporter.cpp
+++ b/Engine/Source/Developer/Datasmith/DatasmithExporter/Private/DatasmithMeshExporter.cpp
@@ -264,23 +264,7 @@ void FDatasmithMeshExporterImpl::CreateDefaultUVs( FDatasmithMesh& Mesh )
 
 void FDatasmithMeshExporterImpl::RegisterStaticMeshAttributes(FMeshDescription& MeshDescription)
 {
-	// This is a local inlining of the UStaticMesh::RegisterMeshAttributes() function, to avoid having to include "Engine" in our dependencies..
-	///////
-	// Add basic vertex attributes
-	MeshDescription.VertexAttributes().RegisterAttribute<FVector>(MeshAttribute::Vertex::Position, 1, FVector::ZeroVector, EMeshAttributeFlags::Lerpable);
-
-	// Add basic vertex instance attributes
-	MeshDescription.VertexInstanceAttributes().RegisterAttribute<FVector2D>(MeshAttribute::VertexInstance::TextureCoordinate, 1, FVector2D::ZeroVector, EMeshAttributeFlags::Lerpable);
-	MeshDescription.VertexInstanceAttributes().RegisterAttribute<FVector>(MeshAttribute::VertexInstance::Normal, 1, FVector::ZeroVector, EMeshAttributeFlags::AutoGenerated);
-	MeshDescription.VertexInstanceAttributes().RegisterAttribute<FVector>(MeshAttribute::VertexInstance::Tangent, 1, FVector::ZeroVector, EMeshAttributeFlags::AutoGenerated);
-	MeshDescription.VertexInstanceAttributes().RegisterAttribute<float>(MeshAttribute::VertexInstance::BinormalSign, 1, 0.0f, EMeshAttributeFlags::AutoGenerated);
-	MeshDescription.VertexInstanceAttributes().RegisterAttribute<FVector4>(MeshAttribute::VertexInstance::Color, 1, FVector4(1.0f, 1.0f, 1.0f, 1.0f), EMeshAttributeFlags::Lerpable);
-
-	// Add basic edge attributes
-	MeshDescription.EdgeAttributes().RegisterAttribute<bool>(MeshAttribute::Edge::IsHard, 1, false);
-
-	// Add basic polygon group attributes
-	MeshDescription.PolygonGroupAttributes().RegisterAttribute<FName>(MeshAttribute::PolygonGroup::ImportedMaterialSlotName); //The unique key to match the mesh material slot
+	FStaticMeshAttributes(MeshDescription).Register();
 }
 
 TSharedPtr<UDatasmithMesh> FDatasmithMeshExporterImpl::GeneratePooledUDatasmithMeshFromFDatasmithMesh(const FDatasmithMesh& Mesh, bool bIsCollisionMesh)
diff --git a/Engine/Source/Developer/MaterialBaking/Private/MaterialRenderItem.cpp b/Engine/Source/Developer/MaterialBaking/Private/MaterialRenderItem.cpp
--- a/Engine/Source/Developer/MaterialBaking/Private/MaterialRenderItem.cpp
+++ b/Engine/Source/Developer/MaterialBaking/Private/MaterialRenderItem.cpp
@@ -186,12 +186,14 @@ void FMeshMaterialRenderItem::PopulateWithMeshData()
 {
 	const FMeshDescription* RawMesh = MeshSettings->RawMeshDescription;
 
-	TVertexAttributesConstRef<FVector> VertexPositions = RawMesh->VertexAttributes().GetAttributesRef<FVector>(MeshAttribute::Vertex::Position);
-	TVertexInstanceAttributesConstRef<FVector> VertexInstanceNormals = RawMesh->VertexInstanceAttributes().GetAttributesRef<FVector>(MeshAttribute::VertexInstance::Normal);
-	TVertexInstanceAttributesConstRef<FVector> VertexInstanceTangents = RawMesh->VertexInstanceAttributes().GetAttributesRef<FVector>(MeshAttribute::VertexInstance::Tangent);
-	TVertexInstanceAttributesConstRef<float> VertexInstanceBinormalSigns = RawMesh->VertexInstanceAttributes().GetAttributesRef<float>(MeshAttribute::VertexInstance::BinormalSign);
-	TVertexInstanceAttributesConstRef<FVector2D> VertexInstanceUVs = RawMesh->VertexInstanceAttributes().GetAttributesRef<FVector2D>(MeshAttribute::VertexInstance::TextureCoordinate);
-	TVertexInstanceAttributesConstRef<FVector4> VertexInstanceColors = RawMesh->VertexInstanceAttributes().GetAttributesRef<FVector4>(MeshAttribute::VertexInstance::Color);
+	FStaticMeshConstAttributes Attributes(*RawMesh);
+	TArrayView<const FVector> VertexPositions = Attributes.GetVertexPositions().GetRawArray();
+	TArrayView<const FVector> VertexInstanceNormals = Attributes.GetVertexInstanceNormals().GetRawArray();
+	TArrayView<const FVector> VertexInstanceTangents = Attributes.GetVertexInstanceTangents().GetRawArray();
+	TArrayView<const float> VertexInstanceBinormalSigns = Attributes.GetVertexInstanceBinormalSigns().GetRawArray();
+	TArrayView<const FVector4> VertexInstanceColors = Attributes.GetVertexInstanceColors().GetRawArray();
+
+	TVertexInstanceAttributesConstRef<FVector2D> VertexInstanceUVs = Attributes.GetVertexInstanceUVs();
 	const int32 NumVerts = RawMesh->Vertices().Num();
 
 	// reserve renderer data
diff --git a/Engine/Source/Developer/MaterialUtilities/Private/MaterialUtilities.cpp b/Engine/Source/Developer/MaterialUtilities/Private/MaterialUtilities.cpp
--- a/Engine/Source/Developer/MaterialUtilities/Private/MaterialUtilities.cpp
+++ b/Engine/Source/Developer/MaterialUtilities/Private/MaterialUtilities.cpp
@@ -1730,7 +1730,7 @@ void FMaterialUtilities::RemapUniqueMaterialIndices(const TArray<FSectionInfo>&
 	{
 		for (int32 LODIndex = 0; LODIndex < MAX_STATIC_MESH_LODS; ++LODIndex)
 		{
-			TVertexAttributesRef<FVector> VertexPositions = InMeshData[MeshIndex].MeshLODData[LODIndex].RawMesh->VertexAttributes().GetAttributesRef<FVector>(MeshAttribute::Vertex::Position);
+			TVertexAttributesRef<FVector> VertexPositions = InMeshData[MeshIndex].MeshLODData[LODIndex].RawMesh->GetVertexPositions();
 
 			if (InMeshData[MeshIndex].bShouldExportLOD[LODIndex])
 			{
@@ -1780,7 +1780,7 @@ void FMaterialUtilities::RemapUniqueMaterialIndices(const TArray<FSectionInfo>&
 	{
 		for (int32 LODIndex = 0; LODIndex < MAX_STATIC_MESH_LODS; ++LODIndex)
 		{
-			TVertexAttributesRef<FVector> VertexPositions = InMeshData[MeshIndex].MeshLODData[LODIndex].RawMesh->VertexAttributes().GetAttributesRef<FVector>(MeshAttribute::Vertex::Position);
+			TVertexAttributesRef<FVector> VertexPositions = InMeshData[MeshIndex].MeshLODData[LODIndex].RawMesh->GetVertexPositions();
 
 			if (InMeshData[MeshIndex].bShouldExportLOD[LODIndex])
 			{
diff --git a/Engine/Source/Developer/MaterialUtilities/Private/MeshRendering.cpp b/Engine/Source/Developer/MaterialUtilities/Private/MeshRendering.cpp
--- a/Engine/Source/Developer/MaterialUtilities/Private/MeshRendering.cpp
+++ b/Engine/Source/Developer/MaterialUtilities/Private/MeshRendering.cpp
@@ -151,12 +151,12 @@ public:
 	{
 		// count triangles for selected material
 		int32 NumTris = 0;
-		for (const FPolygonID PolygonID : RawMesh.Polygons().GetElementIDs())
+		for (const FTriangleID TriangleID : RawMesh.Triangles().GetElementIDs())
 		{
-			const FPolygonGroupID PolygonGroupID = RawMesh.GetPolygonPolygonGroup(PolygonID);
+			const FPolygonGroupID PolygonGroupID = RawMesh.GetTrianglePolygonGroup(TriangleID);
 			if (PolygonGroupID.GetValue() == Data.MaterialIndex)
 			{
-				NumTris += RawMesh.GetPolygonTriangleIDs(PolygonID).Num();
+				NumTris++;
 			}
 		}
 		if (NumTris == 0)
@@ -165,12 +165,13 @@ public:
 			return 0;
 		}
 
-		TVertexAttributesConstRef<FVector> VertexPositions = RawMesh.VertexAttributes().GetAttributesRef<FVector>(MeshAttribute::Vertex::Position);
-		TVertexInstanceAttributesConstRef<FVector> VertexInstanceNormals = RawMesh.VertexInstanceAttributes().GetAttributesRef<FVector>(MeshAttribute::VertexInstance::Normal);
-		TVertexInstanceAttributesConstRef<FVector> VertexInstanceTangents = RawMesh.VertexInstanceAttributes().GetAttributesRef<FVector>(MeshAttribute::VertexInstance::Tangent);
-		TVertexInstanceAttributesConstRef<float> VertexInstanceBinormalSigns = RawMesh.VertexInstanceAttributes().GetAttributesRef<float>(MeshAttribute::VertexInstance::BinormalSign);
-		TVertexInstanceAttributesConstRef<FVector2D> VertexInstanceUVs = RawMesh.VertexInstanceAttributes().GetAttributesRef<FVector2D>(MeshAttribute::VertexInstance::TextureCoordinate);
-		TVertexInstanceAttributesConstRef<FVector4> VertexInstanceColors = RawMesh.VertexInstanceAttributes().GetAttributesRef<FVector4>(MeshAttribute::VertexInstance::Color);
+		FStaticMeshConstAttributes Attributes(RawMesh);
+		TVertexAttributesConstRef<FVector> VertexPositions = Attributes.GetVertexPositions();
+		TVertexInstanceAttributesConstRef<FVector> VertexInstanceNormals = Attributes.GetVertexInstanceNormals();
+		TVertexInstanceAttributesConstRef<FVector> VertexInstanceTangents = Attributes.GetVertexInstanceTangents();
+		TVertexInstanceAttributesConstRef<float> VertexInstanceBinormalSigns = Attributes.GetVertexInstanceBinormalSigns();
+		TVertexInstanceAttributesConstRef<FVector2D> VertexInstanceUVs = Attributes.GetVertexInstanceUVs();
+		TVertexInstanceAttributesConstRef<FVector4> VertexInstanceColors = Attributes.GetVertexInstanceColors();
 
 		int32 NumVerts = NumTris * 3;
 
@@ -200,62 +201,59 @@ public:
 		// add vertices
 		int32 VertIndex = 0;
 		int32 FaceIndex = 0;
-		for (const FPolygonID PolygonID : RawMesh.Polygons().GetElementIDs())
+		for (const FTriangleID TriangleID : RawMesh.Triangles().GetElementIDs())
 		{
-			const FPolygonGroupID PolygonGroupID = RawMesh.GetPolygonPolygonGroup(PolygonID);
-			TArrayView<const FTriangleID> TriangleIDs = RawMesh.GetPolygonTriangleIDs(PolygonID);
-			for (const FTriangleID TriangleID : TriangleIDs)
+			const FPolygonGroupID PolygonGroupID = RawMesh.GetTrianglePolygonGroup(TriangleID);
+
+			if (PolygonGroupID.GetValue() == Data.MaterialIndex)
 			{
-				if (PolygonGroupID.GetValue() == Data.MaterialIndex)
+				for (int32 Corner = 0; Corner < 3; Corner++)
 				{
-					for (int32 Corner = 0; Corner < 3; Corner++)
-					{
-						const int32 SrcVertIndex = FaceIndex * 3 + Corner;
-						const FVertexInstanceID SrcVertexInstanceID = RawMesh.GetTriangleVertexInstance(TriangleID, Corner);
-						const FVertexID SrcVertexID = RawMesh.GetVertexInstanceVertex(SrcVertexInstanceID);
-
-						// add vertex
-						FDynamicMeshVertex* Vert = new(OutVerts)FDynamicMeshVertex();
-						if (!bUseNewUVs)
-						{
-							// compute vertex position from original UV
-							const FVector2D& UV = VertexInstanceUVs.Get(SrcVertexInstanceID, 0);
-							Vert->Position.Set((UV.X - U) * ScaleX, (UV.Y - V) * ScaleY, 0);
-						}
-						else
-						{
-							const FVector2D& UV = Data.TexCoords[SrcVertIndex];
-							Vert->Position.Set(UV.X * ScaleX, UV.Y * ScaleY, 0);
-						}
-						FVector TangentX = VertexInstanceTangents[SrcVertexInstanceID];
-						FVector TangentZ = VertexInstanceNormals[SrcVertexInstanceID];
-						FVector TangentY = FVector::CrossProduct(TangentZ, TangentX).GetSafeNormal() * VertexInstanceBinormalSigns[SrcVertexInstanceID];
-						Vert->SetTangents(TangentX, TangentY, TangentZ);
-						for (int32 TexcoordIndex = 0; TexcoordIndex < NumTexcoords; TexcoordIndex++)
-						{
-							Vert->TextureCoordinate[TexcoordIndex] = VertexInstanceUVs.Get(SrcVertexInstanceID, TexcoordIndex);
-						}
-
-						// Store original vertex positions in texture coordinate data
-						Vert->TextureCoordinate[6].X = VertexPositions[SrcVertexID].X;
-						Vert->TextureCoordinate[6].Y = VertexPositions[SrcVertexID].Y;
-						Vert->TextureCoordinate[7].X = VertexPositions[SrcVertexID].Z;
+					const int32 SrcVertIndex = FaceIndex * 3 + Corner;
+					const FVertexInstanceID SrcVertexInstanceID = RawMesh.GetTriangleVertexInstance(TriangleID, Corner);
+					const FVertexID SrcVertexID = RawMesh.GetVertexInstanceVertex(SrcVertexInstanceID);
 
-						Vert->Color = FLinearColor(VertexInstanceColors[SrcVertexInstanceID]).ToFColor(true);
-						// add index
-						OutIndices.Add(VertIndex);
-						VertIndex++;
+					// add vertex
+					FDynamicMeshVertex* Vert = new(OutVerts)FDynamicMeshVertex();
+					if (!bUseNewUVs)
+					{
+						// compute vertex position from original UV
+						const FVector2D& UV = VertexInstanceUVs.Get(SrcVertexInstanceID, 0);
+						Vert->Position.Set((UV.X - U) * ScaleX, (UV.Y - V) * ScaleY, 0);
+					}
+					else
+					{
+						const FVector2D& UV = Data.TexCoords[SrcVertIndex];
+						Vert->Position.Set(UV.X * ScaleX, UV.Y * ScaleY, 0);
 					}
-					if (bDuplicateTris)
+					FVector TangentX = VertexInstanceTangents[SrcVertexInstanceID];
+					FVector TangentZ = VertexInstanceNormals[SrcVertexInstanceID];
+					FVector TangentY = FVector::CrossProduct(TangentZ, TangentX).GetSafeNormal() * VertexInstanceBinormalSigns[SrcVertexInstanceID];
+					Vert->SetTangents(TangentX, TangentY, TangentZ);
+					for (int32 TexcoordIndex = 0; TexcoordIndex < NumTexcoords; TexcoordIndex++)
 					{
-						// add the same triangle with opposite vertex order
-						OutIndices.Add(VertIndex - 3);
-						OutIndices.Add(VertIndex - 1);
-						OutIndices.Add(VertIndex - 2);
+						Vert->TextureCoordinate[TexcoordIndex] = VertexInstanceUVs.Get(SrcVertexInstanceID, TexcoordIndex);
 					}
+
+					// Store original vertex positions in texture coordinate data
+					Vert->TextureCoordinate[6].X = VertexPositions[SrcVertexID].X;
+					Vert->TextureCoordinate[6].Y = VertexPositions[SrcVertexID].Y;
+					Vert->TextureCoordinate[7].X = VertexPositions[SrcVertexID].Z;
+
+					Vert->Color = FLinearColor(VertexInstanceColors[SrcVertexInstanceID]).ToFColor(true);
+					// add index
+					OutIndices.Add(VertIndex);
+					VertIndex++;
+				}
+				if (bDuplicateTris)
+				{
+					// add the same triangle with opposite vertex order
+					OutIndices.Add(VertIndex - 3);
+					OutIndices.Add(VertIndex - 1);
+					OutIndices.Add(VertIndex - 2);
 				}
-				FaceIndex++;
 			}
+			FaceIndex++;
 		}
 		return NumTris;
 	}
diff --git a/Engine/Source/Developer/MeshBuilder/Private/MeshDescriptionHelper.cpp b/Engine/Source/Developer/MeshBuilder/Private/MeshDescriptionHelper.cpp
--- a/Engine/Source/Developer/MeshBuilder/Private/MeshDescriptionHelper.cpp
+++ b/Engine/Source/Developer/MeshBuilder/Private/MeshDescriptionHelper.cpp
@@ -44,9 +44,11 @@ void FMeshDescriptionHelper::SetupRenderMeshDescription(UObject* Owner, FMeshDes
 	//OutRenderMeshDescription->ComputePolygonTangentsAndNormals(BuildSettings->bRemoveDegenerates ? SMALL_NUMBER : 0.0f);
 
 	FVertexInstanceArray& VertexInstanceArray = RenderMeshDescription.VertexInstances();
-	TVertexInstanceAttributesRef<FVector> Normals = RenderMeshDescription.VertexInstanceAttributes().GetAttributesRef<FVector>( MeshAttribute::VertexInstance::Normal );
-	TVertexInstanceAttributesRef<FVector> Tangents = RenderMeshDescription.VertexInstanceAttributes().GetAttributesRef<FVector>( MeshAttribute::VertexInstance::Tangent );
-	TVertexInstanceAttributesRef<float> BinormalSigns = RenderMeshDescription.VertexInstanceAttributes().GetAttributesRef<float>( MeshAttribute::VertexInstance::BinormalSign );
+
+	FStaticMeshAttributes Attributes(RenderMeshDescription);
+	TVertexInstanceAttributesRef<FVector> Normals = Attributes.GetVertexInstanceNormals();
+	TVertexInstanceAttributesRef<FVector> Tangents = Attributes.GetVertexInstanceTangents();
+	TVertexInstanceAttributesRef<float> BinormalSigns = Attributes.GetVertexInstanceBinormalSigns();
 
 	// Find overlapping corners to accelerate adjacency.
 	FStaticMeshOperations::FindOverlappingCorners(OverlappingCorners, RenderMeshDescription, ComparisonThreshold);
@@ -66,7 +68,7 @@ void FMeshDescriptionHelper::SetupRenderMeshDescription(UObject* Owner, FMeshDes
 
 	if (BuildSettings->bGenerateLightmapUVs && VertexInstanceArray.Num() > 0)
 	{
-		TVertexInstanceAttributesRef<FVector2D> VertexInstanceUVs = RenderMeshDescription.VertexInstanceAttributes().GetAttributesRef<FVector2D>(MeshAttribute::VertexInstance::TextureCoordinate);
+		TVertexInstanceAttributesRef<FVector2D> VertexInstanceUVs = Attributes.GetVertexInstanceUVs();
 		int32 NumIndices = VertexInstanceUVs.GetNumChannels();
 		//Verify the src light map channel
 		if (BuildSettings->SrcLightmapIndex >= NumIndices)
diff --git a/Engine/Source/Developer/MeshBuilder/Private/StaticMeshBuilder.cpp b/Engine/Source/Developer/MeshBuilder/Private/StaticMeshBuilder.cpp
--- a/Engine/Source/Developer/MeshBuilder/Private/StaticMeshBuilder.cpp
+++ b/Engine/Source/Developer/MeshBuilder/Private/StaticMeshBuilder.cpp
@@ -495,23 +495,22 @@ void BuildVertexBuffer(
 {
 	TRACE_CPUPROFILER_EVENT_SCOPE(BuildVertexBuffer);
 
-	const FVertexArray& Vertices = MeshDescription.Vertices();
-	const FVertexInstanceArray& VertexInstances = MeshDescription.VertexInstances();
-	const FPolygonGroupArray& PolygonGroupArray = MeshDescription.PolygonGroups();
-	const FPolygonArray& PolygonArray = MeshDescription.Polygons();
 	TArray<int32> RemapVertexInstanceID;
 	// set up vertex buffer elements
-	StaticMeshBuildVertices.Reserve(VertexInstances.GetArraySize());
+	const int32 NumVertexInstances = MeshDescription.VertexInstances().GetArraySize();
+	StaticMeshBuildVertices.Reserve(NumVertexInstances);
 
-	TPolygonGroupAttributesConstRef<FName> PolygonGroupImportedMaterialSlotNames = MeshDescription.PolygonGroupAttributes().GetAttributesRef<FName>(MeshAttribute::PolygonGroup::ImportedMaterialSlotName);
-	TVertexAttributesConstRef<FVector> VertexPositions = MeshDescription.VertexAttributes().GetAttributesRef<FVector>( MeshAttribute::Vertex::Position );
-	TVertexInstanceAttributesConstRef<FVector> VertexInstanceNormals = MeshDescription.VertexInstanceAttributes().GetAttributesRef<FVector>( MeshAttribute::VertexInstance::Normal );
-	TVertexInstanceAttributesConstRef<FVector> VertexInstanceTangents = MeshDescription.VertexInstanceAttributes().GetAttributesRef<FVector>( MeshAttribute::VertexInstance::Tangent );
-	TVertexInstanceAttributesConstRef<float> VertexInstanceBinormalSigns = MeshDescription.VertexInstanceAttributes().GetAttributesRef<float>( MeshAttribute::VertexInstance::BinormalSign );
-	TVertexInstanceAttributesConstRef<FVector4> VertexInstanceColors = MeshDescription.VertexInstanceAttributes().GetAttributesRef<FVector4>( MeshAttribute::VertexInstance::Color );
-	TVertexInstanceAttributesConstRef<FVector2D> VertexInstanceUVs = MeshDescription.VertexInstanceAttributes().GetAttributesRef<FVector2D>( MeshAttribute::VertexInstance::TextureCoordinate );
+	FStaticMeshConstAttributes Attributes(MeshDescription);
 
-	const bool bHasColors = MeshDescription.VertexInstanceAttributes().HasAttribute(MeshAttribute::VertexInstance::Color);
+	TPolygonGroupAttributesConstRef<FName> PolygonGroupImportedMaterialSlotNames = Attributes.GetPolygonGroupMaterialSlotNames();
+	TVertexAttributesConstRef<FVector> VertexPositions = Attributes.GetVertexPositions();
+	TVertexInstanceAttributesConstRef<FVector> VertexInstanceNormals = Attributes.GetVertexInstanceNormals();
+	TVertexInstanceAttributesConstRef<FVector> VertexInstanceTangents = Attributes.GetVertexInstanceTangents();
+	TVertexInstanceAttributesConstRef<float> VertexInstanceBinormalSigns = Attributes.GetVertexInstanceBinormalSigns();
+	TVertexInstanceAttributesConstRef<FVector4> VertexInstanceColors = Attributes.GetVertexInstanceColors();
+	TVertexInstanceAttributesConstRef<FVector2D> VertexInstanceUVs = Attributes.GetVertexInstanceUVs();
+
+	const bool bHasColors = VertexInstanceColors.IsValid();
 
 	const uint32 NumTextureCoord = VertexInstanceUVs.GetNumChannels();
 	const FMatrix ScaleMatrix = FScaleMatrix(BuildSettings.BuildScale3D).Inverse().GetTransposed();
@@ -546,110 +545,102 @@ void BuildVertexBuffer(
 	float VertexComparisonThreshold = BuildSettings.bRemoveDegenerates ? THRESH_POINTS_ARE_SAME : 0.0f;
 
 	int32 WedgeIndex = 0;
-	for (const FPolygonID PolygonID : MeshDescription.Polygons().GetElementIDs())
+	for (const FTriangleID TriangleID : MeshDescription.Triangles().GetElementIDs())
 	{
-		const FPolygonGroupID PolygonGroupID = MeshDescription.GetPolygonPolygonGroup(PolygonID);
+		const FPolygonGroupID PolygonGroupID = MeshDescription.GetTrianglePolygonGroup(TriangleID);
 		const int32 SectionIndex = PolygonGroupToSectionIndex[PolygonGroupID];
 		TArray<uint32>& SectionIndices = OutPerSectionIndices[SectionIndex];
 
-		TArrayView<const FTriangleID> TriangleIDs = MeshDescription.GetPolygonTriangleIDs(PolygonID);
-		uint32 MinIndex = TNumericLimits< uint32 >::Max();
-		uint32 MaxIndex = TNumericLimits< uint32 >::Min();
-		for (int32 TriangleIndex = 0; TriangleIndex < TriangleIDs.Num(); ++TriangleIndex)
+		TArrayView<const FVertexID> VertexIDs = MeshDescription.GetTriangleVertices(TriangleID);
+
+		FVector CornerPositions[3];
+		for (int32 TriVert = 0; TriVert < 3; ++TriVert)
+		{
+			CornerPositions[TriVert] = VertexPositions[VertexIDs[TriVert]];
+		}
+		FOverlappingThresholds OverlappingThresholds;
+		OverlappingThresholds.ThresholdPosition = VertexComparisonThreshold;
+		// Don't process degenerate triangles.
+		if (PointsEqual(CornerPositions[0], CornerPositions[1], OverlappingThresholds)
+			|| PointsEqual(CornerPositions[0], CornerPositions[2], OverlappingThresholds)
+			|| PointsEqual(CornerPositions[1], CornerPositions[2], OverlappingThresholds))
 		{
-			const FTriangleID TriangleID = TriangleIDs[TriangleIndex];
+			WedgeIndex += 3;
+			continue;
+		}
 
-			FVector CornerPositions[3];
-			for (int32 TriVert = 0; TriVert < 3; ++TriVert)
+		TArrayView<const FVertexInstanceID> VertexInstanceIDs = MeshDescription.GetTriangleVertexInstances(TriangleID);
+		for (int32 TriVert = 0; TriVert < 3; ++TriVert, ++WedgeIndex)
+		{
+			const FVertexInstanceID VertexInstanceID = VertexInstanceIDs[TriVert];
+			const FVector& VertexPosition = CornerPositions[TriVert];
+			const FVector& VertexInstanceNormal = VertexInstanceNormals[VertexInstanceID];
+			const FVector& VertexInstanceTangent = VertexInstanceTangents[VertexInstanceID];
+			const float VertexInstanceBinormalSign = VertexInstanceBinormalSigns[VertexInstanceID];
+
+			FStaticMeshBuildVertex StaticMeshVertex;
+
+			StaticMeshVertex.Position = VertexPosition * BuildSettings.BuildScale3D;
+			StaticMeshVertex.TangentX = ScaleMatrix.TransformVector(VertexInstanceTangent).GetSafeNormal();
+			StaticMeshVertex.TangentY = ScaleMatrix.TransformVector(FVector::CrossProduct(VertexInstanceNormal, VertexInstanceTangent).GetSafeNormal() * VertexInstanceBinormalSign).GetSafeNormal();
+			StaticMeshVertex.TangentZ = ScaleMatrix.TransformVector(VertexInstanceNormal).GetSafeNormal();
+				
+			if (bHasColors)
 			{
-				const FVertexInstanceID VertexInstanceID = MeshDescription.GetTriangleVertexInstance(TriangleID, TriVert);
-				const FVertexID VertexID = MeshDescription.GetVertexInstanceVertex(VertexInstanceID);
-				CornerPositions[TriVert] = VertexPositions[VertexID];
+				const FVector4& VertexInstanceColor = VertexInstanceColors[VertexInstanceID];
+				const FLinearColor LinearColor(VertexInstanceColor);
+				StaticMeshVertex.Color = LinearColor.ToFColor(true);
 			}
-			FOverlappingThresholds OverlappingThresholds;
-			OverlappingThresholds.ThresholdPosition = VertexComparisonThreshold;
-			// Don't process degenerate triangles.
-			if (PointsEqual(CornerPositions[0], CornerPositions[1], OverlappingThresholds)
-				|| PointsEqual(CornerPositions[0], CornerPositions[2], OverlappingThresholds)
-				|| PointsEqual(CornerPositions[1], CornerPositions[2], OverlappingThresholds))
+			else
 			{
-				WedgeIndex += 3;
-				continue;
+				StaticMeshVertex.Color = FColor::White;
 			}
 
-			for (int32 TriVert = 0; TriVert < 3; ++TriVert, ++WedgeIndex)
+			const uint32 MaxNumTexCoords = FMath::Min<int32>(MAX_MESH_TEXTURE_COORDS_MD, MAX_STATIC_TEXCOORDS);
+			for (uint32 UVIndex = 0; UVIndex < MaxNumTexCoords; ++UVIndex)
 			{
-				const FVertexInstanceID VertexInstanceID = MeshDescription.GetTriangleVertexInstance(TriangleID, TriVert);
-				const int32 VertexInstanceValue = VertexInstanceID.GetValue();
-				const FVector& VertexPosition = CornerPositions[TriVert];
-				const FVector& VertexInstanceNormal = VertexInstanceNormals[VertexInstanceID];
-				const FVector& VertexInstanceTangent = VertexInstanceTangents[VertexInstanceID];
-				const float VertexInstanceBinormalSign = VertexInstanceBinormalSigns[VertexInstanceID];
-
-				FStaticMeshBuildVertex StaticMeshVertex;
-
-				StaticMeshVertex.Position = VertexPosition * BuildSettings.BuildScale3D;
-				StaticMeshVertex.TangentX = ScaleMatrix.TransformVector(VertexInstanceTangent).GetSafeNormal();
-				StaticMeshVertex.TangentY = ScaleMatrix.TransformVector(FVector::CrossProduct(VertexInstanceNormal, VertexInstanceTangent).GetSafeNormal() * VertexInstanceBinormalSign).GetSafeNormal();
-				StaticMeshVertex.TangentZ = ScaleMatrix.TransformVector(VertexInstanceNormal).GetSafeNormal();
-				
-				if (bHasColors)
+				if(UVIndex < NumTextureCoord)
 				{
-					const FVector4& VertexInstanceColor = VertexInstanceColors[VertexInstanceID];
-					const FLinearColor LinearColor(VertexInstanceColor);
-					StaticMeshVertex.Color = LinearColor.ToFColor(true);
+					StaticMeshVertex.UVs[UVIndex] = VertexInstanceUVs.Get(VertexInstanceID, UVIndex);
 				}
 				else
 				{
-					StaticMeshVertex.Color = FColor::White;
-				}
-
-				const uint32 MaxNumTexCoords = FMath::Min<int32>(MAX_MESH_TEXTURE_COORDS_MD, MAX_STATIC_TEXCOORDS);
-				for (uint32 UVIndex = 0; UVIndex < MaxNumTexCoords; ++UVIndex)
-				{
-					if(UVIndex < NumTextureCoord)
-					{
-						StaticMeshVertex.UVs[UVIndex] = VertexInstanceUVs.Get(VertexInstanceID, UVIndex);
-					}
-					else
-					{
-						StaticMeshVertex.UVs[UVIndex] = FVector2D(0.0f, 0.0f);
-					}
+					StaticMeshVertex.UVs[UVIndex] = FVector2D(0.0f, 0.0f);
 				}
+			}
 					
 
-				//Never add duplicated vertex instance
-				//Use WedgeIndex since OverlappingCorners has been built based on that
-				const TArray<int32>& DupVerts = OverlappingCorners.FindIfOverlapping(WedgeIndex);
+			//Never add duplicated vertex instance
+			//Use WedgeIndex since OverlappingCorners has been built based on that
+			const TArray<int32>& DupVerts = OverlappingCorners.FindIfOverlapping(WedgeIndex);
 
-				int32 Index = INDEX_NONE;
-				for (int32 k = 0; k < DupVerts.Num(); k++)
+			int32 Index = INDEX_NONE;
+			for (int32 k = 0; k < DupVerts.Num(); k++)
+			{
+				if (DupVerts[k] >= WedgeIndex)
 				{
-					if (DupVerts[k] >= WedgeIndex)
-					{
-						break;
-					}
-					int32 Location = RemapVerts.IsValidIndex(DupVerts[k]) ? RemapVerts[DupVerts[k]] : INDEX_NONE;
-					if (Location != INDEX_NONE && AreVerticesEqual(StaticMeshVertex, StaticMeshBuildVertices[Location], VertexComparisonThreshold))
-					{
-						Index = Location;
-						break;
-					}
+					break;
 				}
-				if (Index == INDEX_NONE)
+				int32 Location = RemapVerts.IsValidIndex(DupVerts[k]) ? RemapVerts[DupVerts[k]] : INDEX_NONE;
+				if (Location != INDEX_NONE && AreVerticesEqual(StaticMeshVertex, StaticMeshBuildVertices[Location], VertexComparisonThreshold))
 				{
-					Index = StaticMeshBuildVertices.Add(StaticMeshVertex);
+					Index = Location;
+					break;
 				}
-				RemapVerts[WedgeIndex] = Index;
-				OutWedgeMap[WedgeIndex] = Index;
-				SectionIndices.Add( Index );
 			}
+			if (Index == INDEX_NONE)
+			{
+				Index = StaticMeshBuildVertices.Add(StaticMeshVertex);
+			}
+			RemapVerts[WedgeIndex] = Index;
+			OutWedgeMap[WedgeIndex] = Index;
+			SectionIndices.Add( Index );
 		}
 	}
 
 
 	//Optimize before setting the buffer
-	if (VertexInstances.GetArraySize() < 100000 * 3)
+	if (NumVertexInstances < 100000 * 3)
 	{
 		BuildOptimizationHelper::CacheOptimizeVertexAndIndexBuffer(StaticMeshBuildVertices, OutPerSectionIndices, OutWedgeMap);
 		//check(OutWedgeMap.Num() == MeshDescription->VertexInstances().Num());
diff --git a/Engine/Source/Developer/MeshMergeUtilities/Private/MeshMergeDataTracker.cpp b/Engine/Source/Developer/MeshMergeUtilities/Private/MeshMergeDataTracker.cpp
--- a/Engine/Source/Developer/MeshMergeUtilities/Private/MeshMergeDataTracker.cpp
+++ b/Engine/Source/Developer/MeshMergeUtilities/Private/MeshMergeDataTracker.cpp
@@ -250,8 +250,9 @@ void FMeshMergeDataTracker::ProcessRawMeshes()
 		const int32 LODIndex = Key.GetLODIndex();
 		const FMeshDescription& RawMesh = MeshPair.Value;
 
-		TVertexInstanceAttributesConstRef<FVector4> VertexInstanceColors = RawMesh.VertexInstanceAttributes().GetAttributesRef<FVector4>(MeshAttribute::VertexInstance::Color);
-		TVertexInstanceAttributesConstRef<FVector2D> VertexInstanceUVs = RawMesh.VertexInstanceAttributes().GetAttributesRef<FVector2D>(MeshAttribute::VertexInstance::TextureCoordinate);
+		FStaticMeshConstAttributes Attributes(RawMesh);
+		TVertexInstanceAttributesConstRef<FVector4> VertexInstanceColors = Attributes.GetVertexInstanceColors();
+		TVertexInstanceAttributesConstRef<FVector2D> VertexInstanceUVs = Attributes.GetVertexInstanceUVs();
 
 		// hash vertex color buffer so we can see if instances have unique vertex data
 		if(VertexInstanceColors.GetNumElements() > 0)
diff --git a/Engine/Source/Developer/MeshMergeUtilities/Private/MeshMergeHelpers.cpp b/Engine/Source/Developer/MeshMergeUtilities/Private/MeshMergeHelpers.cpp
--- a/Engine/Source/Developer/MeshMergeUtilities/Private/MeshMergeHelpers.cpp
+++ b/Engine/Source/Developer/MeshMergeUtilities/Private/MeshMergeHelpers.cpp
@@ -283,14 +283,15 @@ void FMeshMergeHelpers::RetrieveMesh(USkeletalMeshComponent* SkeletalMeshCompone
 		//Empty the raw mesh
 		RawMesh.Empty();
 
-		TVertexAttributesRef<FVector> VertexPositions = RawMesh.VertexAttributes().GetAttributesRef<FVector>(MeshAttribute::Vertex::Position);
-		TEdgeAttributesRef<bool> EdgeHardnesses = RawMesh.EdgeAttributes().GetAttributesRef<bool>(MeshAttribute::Edge::IsHard);
-		TPolygonGroupAttributesRef<FName> PolygonGroupImportedMaterialSlotNames = RawMesh.PolygonGroupAttributes().GetAttributesRef<FName>(MeshAttribute::PolygonGroup::ImportedMaterialSlotName);
-		TVertexInstanceAttributesRef<FVector> VertexInstanceNormals = RawMesh.VertexInstanceAttributes().GetAttributesRef<FVector>(MeshAttribute::VertexInstance::Normal);
-		TVertexInstanceAttributesRef<FVector> VertexInstanceTangents = RawMesh.VertexInstanceAttributes().GetAttributesRef<FVector>(MeshAttribute::VertexInstance::Tangent);
-		TVertexInstanceAttributesRef<float> VertexInstanceBinormalSigns = RawMesh.VertexInstanceAttributes().GetAttributesRef<float>(MeshAttribute::VertexInstance::BinormalSign);
-		TVertexInstanceAttributesRef<FVector4> VertexInstanceColors = RawMesh.VertexInstanceAttributes().GetAttributesRef<FVector4>(MeshAttribute::VertexInstance::Color);
-		TVertexInstanceAttributesRef<FVector2D> VertexInstanceUVs = RawMesh.VertexInstanceAttributes().GetAttributesRef<FVector2D>(MeshAttribute::VertexInstance::TextureCoordinate);
+		FStaticMeshAttributes Attributes(RawMesh);
+		TVertexAttributesRef<FVector> VertexPositions = Attributes.GetVertexPositions();
+		TEdgeAttributesRef<bool> EdgeHardnesses = Attributes.GetEdgeHardnesses();
+		TPolygonGroupAttributesRef<FName> PolygonGroupImportedMaterialSlotNames = Attributes.GetPolygonGroupMaterialSlotNames();
+		TVertexInstanceAttributesRef<FVector> VertexInstanceNormals = Attributes.GetVertexInstanceNormals();
+		TVertexInstanceAttributesRef<FVector> VertexInstanceTangents = Attributes.GetVertexInstanceTangents();
+		TVertexInstanceAttributesRef<float> VertexInstanceBinormalSigns = Attributes.GetVertexInstanceBinormalSigns();
+		TVertexInstanceAttributesRef<FVector4> VertexInstanceColors = Attributes.GetVertexInstanceColors();
+		TVertexInstanceAttributesRef<FVector2D> VertexInstanceUVs = Attributes.GetVertexInstanceUVs();
 
 		int32 TotalTriangles = 0;
 		int32 TotalCorners = 0;
@@ -439,14 +440,15 @@ void FMeshMergeHelpers::ExportStaticMeshLOD(const FStaticMeshLODResources& Stati
 		return;
 	}
 
-	TVertexAttributesRef<FVector> VertexPositions = OutRawMesh.VertexAttributes().GetAttributesRef<FVector>(MeshAttribute::Vertex::Position);
-	TEdgeAttributesRef<bool> EdgeHardnesses = OutRawMesh.EdgeAttributes().GetAttributesRef<bool>(MeshAttribute::Edge::IsHard);
-	TPolygonGroupAttributesRef<FName> PolygonGroupImportedMaterialSlotNames = OutRawMesh.PolygonGroupAttributes().GetAttributesRef<FName>(MeshAttribute::PolygonGroup::ImportedMaterialSlotName);
-	TVertexInstanceAttributesRef<FVector> VertexInstanceNormals = OutRawMesh.VertexInstanceAttributes().GetAttributesRef<FVector>(MeshAttribute::VertexInstance::Normal);
-	TVertexInstanceAttributesRef<FVector> VertexInstanceTangents = OutRawMesh.VertexInstanceAttributes().GetAttributesRef<FVector>(MeshAttribute::VertexInstance::Tangent);
-	TVertexInstanceAttributesRef<float> VertexInstanceBinormalSigns = OutRawMesh.VertexInstanceAttributes().GetAttributesRef<float>(MeshAttribute::VertexInstance::BinormalSign);
-	TVertexInstanceAttributesRef<FVector4> VertexInstanceColors = OutRawMesh.VertexInstanceAttributes().GetAttributesRef<FVector4>(MeshAttribute::VertexInstance::Color);
-	TVertexInstanceAttributesRef<FVector2D> VertexInstanceUVs = OutRawMesh.VertexInstanceAttributes().GetAttributesRef<FVector2D>(MeshAttribute::VertexInstance::TextureCoordinate);
+	FStaticMeshAttributes Attributes(OutRawMesh);
+	TVertexAttributesRef<FVector> VertexPositions = Attributes.GetVertexPositions();
+	TEdgeAttributesRef<bool> EdgeHardnesses = Attributes.GetEdgeHardnesses();
+	TPolygonGroupAttributesRef<FName> PolygonGroupImportedMaterialSlotNames = Attributes.GetPolygonGroupMaterialSlotNames();
+	TVertexInstanceAttributesRef<FVector> VertexInstanceNormals = Attributes.GetVertexInstanceNormals();
+	TVertexInstanceAttributesRef<FVector> VertexInstanceTangents = Attributes.GetVertexInstanceTangents();
+	TVertexInstanceAttributesRef<float> VertexInstanceBinormalSigns = Attributes.GetVertexInstanceBinormalSigns();
+	TVertexInstanceAttributesRef<FVector4> VertexInstanceColors = Attributes.GetVertexInstanceColors();
+	TVertexInstanceAttributesRef<FVector2D> VertexInstanceUVs = Attributes.GetVertexInstanceUVs();
 
 	OutRawMesh.ReserveNewVertices(NumVertexPositions);
 	OutRawMesh.ReserveNewVertexInstances(NumWedges);
@@ -556,7 +558,7 @@ bool FMeshMergeHelpers::CheckWrappingUVs(const TArray<FVector2D>& UVs)
 
 bool FMeshMergeHelpers::CheckWrappingUVs(const FMeshDescription& MeshDescription, int32 UVChannelIndex)
 {
-	TVertexInstanceAttributesConstRef<FVector2D> VertexInstanceUVs = MeshDescription.VertexInstanceAttributes().GetAttributesRef<FVector2D>(MeshAttribute::VertexInstance::TextureCoordinate);
+	TVertexInstanceAttributesConstRef<FVector2D> VertexInstanceUVs = FStaticMeshConstAttributes(MeshDescription).GetVertexInstanceUVs();
 	bool bResult = false;
 
 	//Validate the channel, return false if there is an invalid channel index
@@ -617,7 +619,7 @@ void FMeshMergeHelpers::CullTrianglesFromVolumesAndUnderLandscapes(const UWorld*
 		}
 	}
 
-	TVertexAttributesConstRef<FVector> VertexPositions = InOutRawMesh.VertexAttributes().GetAttributesRef<FVector>(MeshAttribute::Vertex::Position);
+	TVertexAttributesConstRef<FVector> VertexPositions = InOutRawMesh.GetVertexPositions();
 
 	TMap<FVertexID, bool> VertexVisible;
 	VertexVisible.Reserve(InOutRawMesh.Vertices().Num());
@@ -672,20 +674,18 @@ void FMeshMergeHelpers::CullTrianglesFromVolumesAndUnderLandscapes(const UWorld*
 
 
 	// We now know which vertices are below the landscape
-	TArray<FPolygonID> PolygonToRemove;
-	for(const FPolygonID PolygonID : InOutRawMesh.Polygons().GetElementIDs())
+	TArray<FTriangleID> TriangleToRemove;
+	for(const FTriangleID TriangleID : InOutRawMesh.Triangles().GetElementIDs())
 	{
 		bool AboveLandscape = false;
-		for (const FTriangleID TriangleID : InOutRawMesh.GetPolygonTriangleIDs(PolygonID))
+		for (int32 Corner = 0; Corner < 3; ++Corner)
 		{
-			for (int32 Corner = 0; Corner < 3; ++Corner)
-			{
-				AboveLandscape |= VertexVisible[InOutRawMesh.GetVertexInstanceVertex(InOutRawMesh.GetTriangleVertexInstance(TriangleID, Corner))];
-			}
+			AboveLandscape |= VertexVisible[InOutRawMesh.GetVertexInstanceVertex(InOutRawMesh.GetTriangleVertexInstance(TriangleID, Corner))];
 		}
+
 		if (!AboveLandscape)
 		{
-			PolygonToRemove.Add(PolygonID);
+			TriangleToRemove.Add(TriangleID);
 		}
 	}
 
@@ -695,9 +695,9 @@ void FMeshMergeHelpers::CullTrianglesFromVolumesAndUnderLandscapes(const UWorld*
 		TArray<FVertexInstanceID> OrphanedVertexInstances;
 		TArray<FPolygonGroupID> OrphanedPolygonGroups;
 		TArray<FVertexID> OrphanedVertices;
-		for (FPolygonID PolygonID : PolygonToRemove)
+		for (FTriangleID TriangleID : TriangleToRemove)
 		{
-			InOutRawMesh.DeletePolygon(PolygonID, &OrphanedEdges, &OrphanedVertexInstances, &OrphanedPolygonGroups);
+			InOutRawMesh.DeleteTriangle(TriangleID, &OrphanedEdges, &OrphanedVertexInstances, &OrphanedPolygonGroups);
 		}
 		//Do not remove the polygongroup since its indexed with the mesh material array
 		/*for (FPolygonGroupID PolygonGroupID : OrphanedPolygonGroups)
@@ -725,29 +725,28 @@ void FMeshMergeHelpers::CullTrianglesFromVolumesAndUnderLandscapes(const UWorld*
 
 void FMeshMergeHelpers::PropagateSplineDeformationToRawMesh(const USplineMeshComponent* InSplineMeshComponent, FMeshDescription &OutRawMesh)
 {
-	TVertexAttributesRef<FVector> VertexPositions = OutRawMesh.VertexAttributes().GetAttributesRef<FVector>(MeshAttribute::Vertex::Position);
-	TVertexInstanceAttributesRef<FVector> VertexInstanceNormals = OutRawMesh.VertexInstanceAttributes().GetAttributesRef<FVector>(MeshAttribute::VertexInstance::Normal);
-	TVertexInstanceAttributesRef<FVector> VertexInstanceTangents = OutRawMesh.VertexInstanceAttributes().GetAttributesRef<FVector>(MeshAttribute::VertexInstance::Tangent);
-	TVertexInstanceAttributesRef<float> VertexInstanceBinormalSigns = OutRawMesh.VertexInstanceAttributes().GetAttributesRef<float>(MeshAttribute::VertexInstance::BinormalSign);
+	FStaticMeshAttributes Attributes(OutRawMesh);
+
+	TVertexAttributesRef<FVector> VertexPositions = Attributes.GetVertexPositions();
+	TVertexInstanceAttributesRef<FVector> VertexInstanceNormals = Attributes.GetVertexInstanceNormals();
+	TVertexInstanceAttributesRef<FVector> VertexInstanceTangents = Attributes.GetVertexInstanceTangents();
+	TVertexInstanceAttributesRef<float> VertexInstanceBinormalSigns = Attributes.GetVertexInstanceBinormalSigns();
 
 	// Apply spline deformation for each vertex's tangents
 	int32 WedgeIndex = 0;
-	for (const FPolygonID PolygonID : OutRawMesh.Polygons().GetElementIDs())
+	for (const FTriangleID TriangleID : OutRawMesh.Triangles().GetElementIDs())
 	{
-		for (const FTriangleID TriangleID : OutRawMesh.GetPolygonTriangleIDs(PolygonID))
+		for (int32 Corner = 0; Corner < 3; ++Corner, ++WedgeIndex)
 		{
-			for (int32 Corner = 0; Corner < 3; ++Corner, ++WedgeIndex)
-			{
-				const FVertexInstanceID VertexInstanceID = OutRawMesh.GetTriangleVertexInstance(TriangleID, Corner);
-				const FVertexID VertexID = OutRawMesh.GetVertexInstanceVertex(VertexInstanceID);
-				const float& AxisValue = USplineMeshComponent::GetAxisValue(VertexPositions[VertexID], InSplineMeshComponent->ForwardAxis);
-				FTransform SliceTransform = InSplineMeshComponent->CalcSliceTransform(AxisValue);
-				FVector TangentY = FVector::CrossProduct(VertexInstanceNormals[VertexInstanceID], VertexInstanceTangents[VertexInstanceID]).GetSafeNormal() * VertexInstanceBinormalSigns[VertexInstanceID];
-				VertexInstanceTangents[VertexInstanceID] = SliceTransform.TransformVector(VertexInstanceTangents[VertexInstanceID]);
-				TangentY = SliceTransform.TransformVector(TangentY);
-				VertexInstanceNormals[VertexInstanceID] = SliceTransform.TransformVector(VertexInstanceNormals[VertexInstanceID]);
-				VertexInstanceBinormalSigns[VertexInstanceID] = GetBasisDeterminantSign(VertexInstanceTangents[VertexInstanceID], TangentY, VertexInstanceNormals[VertexInstanceID]);
-			}
+			const FVertexInstanceID VertexInstanceID = OutRawMesh.GetTriangleVertexInstance(TriangleID, Corner);
+			const FVertexID VertexID = OutRawMesh.GetVertexInstanceVertex(VertexInstanceID);
+			const float& AxisValue = USplineMeshComponent::GetAxisValue(VertexPositions[VertexID], InSplineMeshComponent->ForwardAxis);
+			FTransform SliceTransform = InSplineMeshComponent->CalcSliceTransform(AxisValue);
+			FVector TangentY = FVector::CrossProduct(VertexInstanceNormals[VertexInstanceID], VertexInstanceTangents[VertexInstanceID]).GetSafeNormal() * VertexInstanceBinormalSigns[VertexInstanceID];
+			VertexInstanceTangents[VertexInstanceID] = SliceTransform.TransformVector(VertexInstanceTangents[VertexInstanceID]);
+			TangentY = SliceTransform.TransformVector(TangentY);
+			VertexInstanceNormals[VertexInstanceID] = SliceTransform.TransformVector(VertexInstanceNormals[VertexInstanceID]);
+			VertexInstanceBinormalSigns[VertexInstanceID] = GetBasisDeterminantSign(VertexInstanceTangents[VertexInstanceID], TangentY, VertexInstanceNormals[VertexInstanceID]);
 		}
 	}
 
@@ -794,14 +793,11 @@ void FMeshMergeHelpers::TransformRawMeshVertexData(const FTransform& InTransform
 {
 	TRACE_CPUPROFILER_EVENT_SCOPE(FMeshMergeHelpers::TransformRawMeshVertexData)
 
-	TVertexAttributesRef<FVector> VertexPositions = OutRawMesh.VertexAttributes().GetAttributesRef<FVector>(MeshAttribute::Vertex::Position);
-	TEdgeAttributesRef<bool> EdgeHardnesses = OutRawMesh.EdgeAttributes().GetAttributesRef<bool>(MeshAttribute::Edge::IsHard);
-	TPolygonGroupAttributesRef<FName> PolygonGroupImportedMaterialSlotNames = OutRawMesh.PolygonGroupAttributes().GetAttributesRef<FName>(MeshAttribute::PolygonGroup::ImportedMaterialSlotName);
-	TVertexInstanceAttributesRef<FVector> VertexInstanceNormals = OutRawMesh.VertexInstanceAttributes().GetAttributesRef<FVector>(MeshAttribute::VertexInstance::Normal);
-	TVertexInstanceAttributesRef<FVector> VertexInstanceTangents = OutRawMesh.VertexInstanceAttributes().GetAttributesRef<FVector>(MeshAttribute::VertexInstance::Tangent);
-	TVertexInstanceAttributesRef<float> VertexInstanceBinormalSigns = OutRawMesh.VertexInstanceAttributes().GetAttributesRef<float>(MeshAttribute::VertexInstance::BinormalSign);
-	TVertexInstanceAttributesRef<FVector4> VertexInstanceColors = OutRawMesh.VertexInstanceAttributes().GetAttributesRef<FVector4>(MeshAttribute::VertexInstance::Color);
-	TVertexInstanceAttributesRef<FVector2D> VertexInstanceUVs = OutRawMesh.VertexInstanceAttributes().GetAttributesRef<FVector2D>(MeshAttribute::VertexInstance::TextureCoordinate);
+	FStaticMeshAttributes Attributes(OutRawMesh);
+	TArrayView<FVector> VertexPositions = Attributes.GetVertexPositions().GetRawArray();
+	TArrayView<FVector> VertexInstanceNormals = Attributes.GetVertexInstanceNormals().GetRawArray();
+	TArrayView<FVector> VertexInstanceTangents = Attributes.GetVertexInstanceTangents().GetRawArray();
+	TArrayView<float> VertexInstanceBinormalSigns = Attributes.GetVertexInstanceBinormalSigns().GetRawArray();
 
 	for(const FVertexID& VertexID : OutRawMesh.Vertices().GetElementIDs())
 	{
@@ -1040,7 +1036,7 @@ FVector2D FMeshMergeHelpers::GetValidUV(const FVector2D& UV)
 
 void FMeshMergeHelpers::CalculateTextureCoordinateBoundsForRawMesh(const FMeshDescription& InRawMesh, TArray<FBox2D>& OutBounds)
 {
-	TVertexInstanceAttributesConstRef<FVector2D> VertexInstanceUVs = InRawMesh.VertexInstanceAttributes().GetAttributesRef<FVector2D>(MeshAttribute::VertexInstance::TextureCoordinate);
+	TVertexInstanceAttributesConstRef<FVector2D> VertexInstanceUVs = FStaticMeshConstAttributes(InRawMesh).GetVertexInstanceUVs();
 	OutBounds.Empty();
 	for (const FTriangleID TriangleID : InRawMesh.Polygons().GetElementIDs())
 	{
@@ -1077,7 +1073,7 @@ bool FMeshMergeHelpers::PropagatePaintedColorsToRawMesh(const UStaticMeshCompone
 			const int32 NumRenderWedges = RenderModel.IndexBuffer.GetNumIndices();
 			const bool bUseRenderWedges = NumWedges == NumRenderWedges;
 
-			TVertexInstanceAttributesRef<FVector4> VertexInstanceColors = RawMesh.VertexInstanceAttributes().GetAttributesRef<FVector4>(MeshAttribute::VertexInstance::Color);
+			TVertexInstanceAttributesRef<FVector4> VertexInstanceColors = FStaticMeshAttributes(RawMesh).GetVertexInstanceColors();
 
 			if (bUseRenderWedges)
 			{
@@ -1085,15 +1081,11 @@ bool FMeshMergeHelpers::PropagatePaintedColorsToRawMesh(const UStaticMeshCompone
 				TMap<int32, FVertexInstanceID> IndexToVertexInstanceID;
 				IndexToVertexInstanceID.Reserve(NumWedges);
 				int32 CurrentWedgeIndex = 0;
-				for (const FPolygonID PolygonID : RawMesh.Polygons().GetElementIDs())
+				for (const FTriangleID TriangleID : RawMesh.Triangles().GetElementIDs())
 				{
-					TArrayView<const FTriangleID> TriangleIDs = RawMesh.GetPolygonTriangleIDs(PolygonID);
-					for (const FTriangleID TriangleID : TriangleIDs)
+					for (int32 Corner = 0; Corner < 3; ++Corner, ++CurrentWedgeIndex)
 					{
-						for (int32 Corner = 0; Corner < 3; ++Corner, ++CurrentWedgeIndex)
-						{
-							IndexToVertexInstanceID.Add(CurrentWedgeIndex, RawMesh.GetTriangleVertexInstance(TriangleID, Corner));
-						}
+						IndexToVertexInstanceID.Add(CurrentWedgeIndex, RawMesh.GetTriangleVertexInstance(TriangleID, Corner));
 					}
 				}
 
@@ -1135,7 +1127,7 @@ bool FMeshMergeHelpers::PropagatePaintedColorsToRawMesh(const UStaticMeshCompone
 						{
 							WedgeColor = ColorVertexBuffer.VertexColor(VertIndex);
 						}
-						TArrayView<const FVertexInstanceID> VertexInstances = RawMesh.GetVertexVertexInstances(VertexID);
+						TArrayView<const FVertexInstanceID> VertexInstances = RawMesh.GetVertexVertexInstanceIDs(VertexID);
 						for (const FVertexInstanceID& VertexInstanceID : VertexInstances)
 						{
 							VertexInstanceColors[VertexInstanceID] = FLinearColor(WedgeColor);
@@ -1365,7 +1357,7 @@ void FMeshMergeHelpers::MergeImpostersToRawMesh(TArray<const UStaticMeshComponen
 		else if (!InPivot.IsZero())
 		{
 			// Apply pivot offset if non null
-			TVertexAttributesRef<FVector> ImposterMeshVertexPositions = ImposterMesh.VertexAttributes().GetAttributesRef<FVector>(MeshAttribute::Vertex::Position);
+			TVertexAttributesRef<FVector> ImposterMeshVertexPositions = ImposterMesh.GetVertexPositions();
 			for (FVertexID VertexID : ImposterMesh.Vertices().GetElementIDs())
 			{
 				ImposterMeshVertexPositions[VertexID] -= InPivot;
@@ -1373,7 +1365,9 @@ void FMeshMergeHelpers::MergeImpostersToRawMesh(TArray<const UStaticMeshComponen
 		}
 
 		TPolygonGroupAttributesRef<FName> SourcePolygonGroupImportedMaterialSlotNames = ImposterMeshAttributes.GetPolygonGroupMaterialSlotNames();
-		TPolygonGroupAttributesRef<FName> TargetPolygonGroupImportedMaterialSlotNames = InRawMesh.PolygonGroupAttributes().GetAttributesRef<FName>(MeshAttribute::PolygonGroup::ImportedMaterialSlotName);
+
+		FStaticMeshAttributes Attributes(InRawMesh);
+		TPolygonGroupAttributesRef<FName> TargetPolygonGroupImportedMaterialSlotNames = Attributes.GetPolygonGroupMaterialSlotNames();
 
 		//Add the missing polygon group ID to the target(InRawMesh)
 		//Remap the source mesh(ImposterMesh) polygongroup to fit with the target polygon groups
diff --git a/Engine/Source/Developer/MeshMergeUtilities/Private/MeshMergeUtilities.cpp b/Engine/Source/Developer/MeshMergeUtilities/Private/MeshMergeUtilities.cpp
--- a/Engine/Source/Developer/MeshMergeUtilities/Private/MeshMergeUtilities.cpp
+++ b/Engine/Source/Developer/MeshMergeUtilities/Private/MeshMergeUtilities.cpp
@@ -231,7 +231,7 @@ void FMeshMergeUtilities::BakeMaterialsForComponent(TArray<TWeakObjectPtr<UObjec
 					const bool bNeedsUniqueUVs = FMeshMergeHelpers::CheckWrappingUVs(*(MeshSettings.RawMeshDescription), MeshSettings.TextureCoordinateIndex);
 					const int32 LightMapUVIndex = Adapter->LightmapUVIndex();
 					
-					TVertexInstanceAttributesConstRef<FVector2D> VertexInstanceUVs = MeshSettings.RawMeshDescription->VertexInstanceAttributes().GetAttributesRef<FVector2D>(MeshAttribute::VertexInstance::TextureCoordinate);
+					TVertexInstanceAttributesConstRef<FVector2D> VertexInstanceUVs = FStaticMeshAttributes(*MeshSettings.RawMeshDescription).GetVertexInstanceUVs();
 					if (bNeedsUniqueUVs && MeshSettings.TextureCoordinateIndex != LightMapUVIndex && VertexInstanceUVs.GetNumElements() > 0 && VertexInstanceUVs.GetNumChannels() > LightMapUVIndex)
 					{
 						MeshSettings.TextureCoordinateIndex = LightMapUVIndex;
@@ -1386,7 +1386,7 @@ void FMeshMergeUtilities::CreateProxyMesh(const TArray<UStaticMeshComponent*>& I
 						FMeshData MeshSettings;
 						MeshSettings.RawMeshDescription = MeshDescription;
 
-						TVertexInstanceAttributesRef<FVector2D> VertexInstanceUVs = MeshDescription->VertexInstanceAttributes().GetAttributesRef<FVector2D>(MeshAttribute::VertexInstance::TextureCoordinate);
+						TVertexInstanceAttributesRef<FVector2D> VertexInstanceUVs = FStaticMeshAttributes(*MeshDescription).GetVertexInstanceUVs();
 
 						// If we already have lightmap uvs generated and they are valid, we can reuse those instead of having to generate new ones
 						const int32 LightMapCoordinateIndex = StaticMeshComponent->GetStaticMesh()->LightMapCoordinateIndex;
@@ -1634,7 +1634,7 @@ void FMeshMergeUtilities::CreateProxyMesh(const TArray<UStaticMeshComponent*>& I
 				}
 				else
 				{
-					TVertexInstanceAttributesRef<FVector2D> VertexInstanceUVs = MeshData->RawMeshDescription->VertexInstanceAttributes().GetAttributesRef<FVector2D>(MeshAttribute::VertexInstance::TextureCoordinate);
+					TVertexInstanceAttributesRef<FVector2D> VertexInstanceUVs = FStaticMeshAttributes(*MeshData->RawMeshDescription).GetVertexInstanceUVs();
 					MergeData.NewUVs.Reset(MeshData->RawMeshDescription->VertexInstances().Num());
 					for (const FVertexInstanceID VertexInstanceID : MeshData->RawMeshDescription->VertexInstances().GetElementIDs())
 					{
@@ -2179,7 +2179,7 @@ void FMeshMergeUtilities::MergeComponentsToStaticMesh(const TArray<UPrimitiveCom
 						// Check if there are lightmap uvs available?
 						const int32 LightMapUVIndex = StaticMeshComponentsToMerge[Key.GetMeshIndex()]->GetStaticMesh()->LightMapCoordinateIndex;
 
-						TVertexInstanceAttributesRef<FVector2D> VertexInstanceUVs = MeshData.RawMeshDescription->VertexInstanceAttributes().GetAttributesRef<FVector2D>(MeshAttribute::VertexInstance::TextureCoordinate);
+						TVertexInstanceAttributesRef<FVector2D> VertexInstanceUVs = FStaticMeshAttributes(*MeshData.RawMeshDescription).GetVertexInstanceUVs();
 						if (InSettings.bReuseMeshLightmapUVs && VertexInstanceUVs.GetNumElements() > 0 && VertexInstanceUVs.GetNumChannels() > LightMapUVIndex)
 						{
 							MeshData.TextureCoordinateIndex = LightMapUVIndex;
@@ -2442,7 +2442,7 @@ void FMeshMergeUtilities::MergeComponentsToStaticMesh(const TArray<UPrimitiveCom
 
 				if(RawMesh.Vertices().Num())
 				{
-					const TVertexInstanceAttributesRef<const FVector2D> VertexInstanceUVs = RawMesh.VertexInstanceAttributes().GetAttributesRef<FVector2D>(MeshAttribute::VertexInstance::TextureCoordinate);
+					TVertexInstanceAttributesRef<const FVector2D> VertexInstanceUVs = FStaticMeshConstAttributes(RawMesh).GetVertexInstanceUVs();
 					MergedMatUVChannel = FMath::Max(MergedMatUVChannel, VertexInstanceUVs.GetNumChannels());
 				}
 			}
@@ -2466,7 +2466,7 @@ void FMeshMergeUtilities::MergeComponentsToStaticMesh(const TArray<UPrimitiveCom
 					ProcessedMaterials.Add(MaterialIndex);
 					if (RawMesh->Vertices().Num())
 					{
-						TVertexInstanceAttributesRef<FVector2D> VertexInstanceUVs = RawMesh->VertexInstanceAttributes().GetAttributesRef<FVector2D>(MeshAttribute::VertexInstance::TextureCoordinate);
+						TVertexInstanceAttributesRef<FVector2D> VertexInstanceUVs = FStaticMeshAttributes(*RawMesh).GetVertexInstanceUVs();
 						int32 NumUVChannel = FMath::Min(VertexInstanceUVs.GetNumChannels(), (int32)MAX_MESH_TEXTURE_COORDS);
 						for (int32 UVChannelIdx = 0; UVChannelIdx < NumUVChannel; ++UVChannelIdx)
 						{
@@ -2578,7 +2578,7 @@ void FMeshMergeUtilities::MergeComponentsToStaticMesh(const TArray<UPrimitiveCom
 		for (FPolygonGroupID PolygonGroupID : TargetRawMesh.PolygonGroups().GetElementIDs())
 		{
 			//Skip empty group
-			if (TargetRawMesh.GetPolygonGroupPolygons(PolygonGroupID).Num() > 0)
+			if (TargetRawMesh.GetPolygonGroupPolygonIDs(PolygonGroupID).Num() > 0)
 			{
 				if (PolygonGroupID.GetValue() < DataTracker.NumberOfUniqueSections())
 				{
@@ -2972,7 +2972,7 @@ void FMeshMergeUtilities::CreateMergedRawMeshes(FMeshMergeDataTracker& InDataTra
 			TArray<FPolygonGroupID> PolygonGroupToRemove;
 			for (FPolygonGroupID PolygonGroupID : MergedMesh.PolygonGroups().GetElementIDs())
 			{
-				if (MergedMesh.GetPolygonGroupPolygons(PolygonGroupID).Num() < 1)
+				if (MergedMesh.GetPolygonGroupPolygonIDs(PolygonGroupID).Num() < 1)
 				{
 					PolygonGroupToRemove.Add(PolygonGroupID);
 					
diff --git a/Engine/Source/Developer/MeshMergeUtilities/Private/ProxyGenerationProcessor.cpp b/Engine/Source/Developer/MeshMergeUtilities/Private/ProxyGenerationProcessor.cpp
--- a/Engine/Source/Developer/MeshMergeUtilities/Private/ProxyGenerationProcessor.cpp
+++ b/Engine/Source/Developer/MeshMergeUtilities/Private/ProxyGenerationProcessor.cpp
@@ -206,14 +206,14 @@ void FProxyGenerationProcessor::ProcessJob(const FGuid& JobGuid, FProxyGeneratio
 	const bool bContainsImposters = Data->MergeData->ImposterComponents.Num() > 0;
 	FBox ImposterBounds(EForceInit::ForceInit);
 
-	TPolygonGroupAttributesConstRef<FName> PolygonGroupMaterialSlotName = Data->RawMesh.PolygonGroupAttributes().GetAttributesRef<FName>(MeshAttribute::PolygonGroup::ImportedMaterialSlotName);
+	TPolygonGroupAttributesConstRef<FName> PolygonGroupMaterialSlotName = FStaticMeshAttributes(Data->RawMesh).GetPolygonGroupMaterialSlotNames();
 
 	auto RemoveVertexColorAndCommitMeshDescription = [&StaticMesh, &Data, &ProxyMaterial, &PolygonGroupMaterialSlotName]()
 	{
 		if (!Data->MergeData->InProxySettings.bAllowVertexColors)
 		{
 			//We cannot remove the vertex color with the mesh description so we assign a white value to all color
-			TVertexInstanceAttributesRef<FVector4> VertexInstanceColors = Data->RawMesh.VertexInstanceAttributes().GetAttributesRef<FVector4>(MeshAttribute::VertexInstance::Color);
+			TVertexInstanceAttributesRef<FVector4> VertexInstanceColors = FStaticMeshAttributes(Data->RawMesh).GetVertexInstanceColors();
 			//set all value to white
 			for (const FVertexInstanceID& VertexInstanceID : Data->RawMesh.VertexInstances().GetElementIDs())
 			{
diff --git a/Engine/Source/Developer/MeshSimplifier/Private/QuadricMeshReduction.cpp b/Engine/Source/Developer/MeshSimplifier/Private/QuadricMeshReduction.cpp
--- a/Engine/Source/Developer/MeshSimplifier/Private/QuadricMeshReduction.cpp
+++ b/Engine/Source/Developer/MeshSimplifier/Private/QuadricMeshReduction.cpp
@@ -110,130 +110,124 @@ public:
 		float SurfaceArea = 0.0f;
 
 		int32 WedgeIndex = 0;
-		for (const FPolygonID& PolygonID : InMesh.Polygons().GetElementIDs())
+		for (const FTriangleID TriangleID : InMesh.Triangles().GetElementIDs())
 		{
-			const FPolygonGroupID PolygonGroupID = InMesh.GetPolygonPolygonGroup(PolygonID);
+			const FPolygonGroupID PolygonGroupID = InMesh.GetTrianglePolygonGroup(TriangleID);
+			TArrayView<const FVertexID> VertexIDs = InMesh.GetTriangleVertices(TriangleID);
 
-			TArrayView<const FTriangleID> TriangleIDs = InMesh.GetPolygonTriangleIDs(PolygonID);
-			for (int32 TriangleIndex = 0; TriangleIndex < TriangleIDs.Num(); ++TriangleIndex)
+			FVector CornerPositions[3];
+			for (int32 TriVert = 0; TriVert < 3; ++TriVert)
 			{
-				const FTriangleID TriangleID = TriangleIDs[TriangleIndex];
+				const FVertexID TmpVertexID = VertexIDs[TriVert];
+				const FVertexID VertexID = bWeldVertices ? VertexIDRemap[TmpVertexID] : TmpVertexID;
+				CornerPositions[TriVert] = InVertexPositions[VertexID];
+			}
 
-				FVector CornerPositions[3];
-				for (int32 TriVert = 0; TriVert < 3; ++TriVert)
-				{
-					const FVertexInstanceID VertexInstanceID = InMesh.GetTriangleVertexInstance(TriangleID, TriVert);
-					const FVertexID TmpVertexID = InMesh.GetVertexInstanceVertex(VertexInstanceID);
-					const FVertexID VertexID = bWeldVertices ? VertexIDRemap[TmpVertexID] : TmpVertexID;
-					CornerPositions[TriVert] = InVertexPositions[VertexID];
-				}
+			// Don't process degenerate triangles.
+			if( PointsEqual(CornerPositions[0], CornerPositions[1]) ||
+				PointsEqual(CornerPositions[0], CornerPositions[2]) ||
+				PointsEqual(CornerPositions[1], CornerPositions[2]) )
+			{
+				WedgeIndex += 3;
+				continue;
+			}
 
-				// Don't process degenerate triangles.
-				if( PointsEqual(CornerPositions[0], CornerPositions[1]) ||
-					PointsEqual(CornerPositions[0], CornerPositions[2]) ||
-					PointsEqual(CornerPositions[1], CornerPositions[2]) )
-				{
-					WedgeIndex += 3;
-					continue;
-				}
+			int32 VertexIndices[3];
+			for (int32 TriVert = 0; TriVert < 3; ++TriVert, ++WedgeIndex)
+			{
+				const FVertexInstanceID VertexInstanceID = InMesh.GetTriangleVertexInstance(TriangleID, TriVert);
+				const int32 VertexInstanceValue = VertexInstanceID.GetValue();
+				const FVector& VertexPosition = CornerPositions[TriVert];
+
+				TVertSimp< NumTexCoords > NewVert;
 
-				int32 VertexIndices[3];
-				for (int32 TriVert = 0; TriVert < 3; ++TriVert, ++WedgeIndex)
+				NewVert.Position = CornerPositions[TriVert];
+				NewVert.Tangents[0] = InVertexTangents[ VertexInstanceID ];
+				NewVert.Normal = InVertexNormals[ VertexInstanceID ];
+				NewVert.Tangents[1] = FVector(0.0f);
+				if (!NewVert.Normal.IsNearlyZero(SMALL_NUMBER) && !NewVert.Tangents[0].IsNearlyZero(SMALL_NUMBER))
 				{
-					const FVertexInstanceID VertexInstanceID = InMesh.GetTriangleVertexInstance(TriangleID, TriVert);
-					const int32 VertexInstanceValue = VertexInstanceID.GetValue();
-					const FVector& VertexPosition = CornerPositions[TriVert];
+					NewVert.Tangents[1] = FVector::CrossProduct(NewVert.Normal, NewVert.Tangents[0]).GetSafeNormal() * InVertexBinormalSigns[ VertexInstanceID ];
+				}
 
-					TVertSimp< NumTexCoords > NewVert;
+				// Fix bad tangents
+				NewVert.Tangents[0] = NewVert.Tangents[0].ContainsNaN() ? FVector::ZeroVector : NewVert.Tangents[0];
+				NewVert.Tangents[1] = NewVert.Tangents[1].ContainsNaN() ? FVector::ZeroVector : NewVert.Tangents[1];
+				NewVert.Normal = NewVert.Normal.ContainsNaN() ? FVector::ZeroVector : NewVert.Normal;
+				NewVert.Color = FLinearColor(InVertexColors[ VertexInstanceID ]);
 
-					NewVert.Position = CornerPositions[TriVert];
-					NewVert.Tangents[0] = InVertexTangents[ VertexInstanceID ];
-					NewVert.Normal = InVertexNormals[ VertexInstanceID ];
-					NewVert.Tangents[1] = FVector(0.0f);
-					if (!NewVert.Normal.IsNearlyZero(SMALL_NUMBER) && !NewVert.Tangents[0].IsNearlyZero(SMALL_NUMBER))
+				for (int32 UVIndex = 0; UVIndex < NumTexCoords; UVIndex++)
+				{
+					if (UVIndex < InVertexUVs.GetNumChannels())
 					{
-						NewVert.Tangents[1] = FVector::CrossProduct(NewVert.Normal, NewVert.Tangents[0]).GetSafeNormal() * InVertexBinormalSigns[ VertexInstanceID ];
+						NewVert.TexCoords[UVIndex] = InVertexUVs.Get(VertexInstanceID, UVIndex);
+						InMeshNumTexCoords = FMath::Max(UVIndex + 1, InMeshNumTexCoords);
 					}
-
-					// Fix bad tangents
-					NewVert.Tangents[0] = NewVert.Tangents[0].ContainsNaN() ? FVector::ZeroVector : NewVert.Tangents[0];
-					NewVert.Tangents[1] = NewVert.Tangents[1].ContainsNaN() ? FVector::ZeroVector : NewVert.Tangents[1];
-					NewVert.Normal = NewVert.Normal.ContainsNaN() ? FVector::ZeroVector : NewVert.Normal;
-					NewVert.Color = FLinearColor(InVertexColors[ VertexInstanceID ]);
-
-					for (int32 UVIndex = 0; UVIndex < NumTexCoords; UVIndex++)
+					else
 					{
-						if (UVIndex < InVertexUVs.GetNumChannels())
-						{
-							NewVert.TexCoords[UVIndex] = InVertexUVs.Get(VertexInstanceID, UVIndex);
-							InMeshNumTexCoords = FMath::Max(UVIndex + 1, InMeshNumTexCoords);
-						}
-						else
-						{
-							NewVert.TexCoords[UVIndex] = FVector2D::ZeroVector;
-						}
+						NewVert.TexCoords[UVIndex] = FVector2D::ZeroVector;
 					}
+				}
 
-					// Make sure this vertex is valid from the start
-					NewVert.Correct();
+				// Make sure this vertex is valid from the start
+				NewVert.Correct();
 					
 
-					//Never add duplicated vertex instance
-					//Use WedgeIndex since OverlappingCorners has been built based on that
-					const TArray<int32>& DupVerts = InOverlappingCorners.FindIfOverlapping(WedgeIndex);
+				//Never add duplicated vertex instance
+				//Use WedgeIndex since OverlappingCorners has been built based on that
+				const TArray<int32>& DupVerts = InOverlappingCorners.FindIfOverlapping(WedgeIndex);
 
-					int32 Index = INDEX_NONE;
-					for (int32 k = 0; k < DupVerts.Num(); k++)
+				int32 Index = INDEX_NONE;
+				for (int32 k = 0; k < DupVerts.Num(); k++)
+				{
+					if (DupVerts[k] >= WedgeIndex)
 					{
-						if (DupVerts[k] >= WedgeIndex)
-						{
-							// the verts beyond me haven't been placed yet, so these duplicates are not relevant
-							break;
-						}
+						// the verts beyond me haven't been placed yet, so these duplicates are not relevant
+						break;
+					}
 
-						int32* Location = VertsMap.Find(DupVerts[k]);
-						if (Location)
-						{
-							TVertSimp< NumTexCoords >& FoundVert = Verts[*Location];
+					int32* Location = VertsMap.Find(DupVerts[k]);
+					if (Location)
+					{
+						TVertSimp< NumTexCoords >& FoundVert = Verts[*Location];
 
-							if (NewVert.Equals(FoundVert))
-							{
-								Index = *Location;
-								break;
-							}
+						if (NewVert.Equals(FoundVert))
+						{
+							Index = *Location;
+							break;
 						}
 					}
-					if (Index == INDEX_NONE)
-					{
-						Index = Verts.Add(NewVert);
-						VertsMap.Add(WedgeIndex, Index);
-					}
-					VertexIndices[TriVert] = Index;
 				}
-				
-				// Reject degenerate triangles.
-				if (VertexIndices[0] == VertexIndices[1] ||
-					VertexIndices[1] == VertexIndices[2] ||
-					VertexIndices[0] == VertexIndices[2])
+				if (Index == INDEX_NONE)
 				{
-					continue;
+					Index = Verts.Add(NewVert);
+					VertsMap.Add(WedgeIndex, Index);
 				}
+				VertexIndices[TriVert] = Index;
+			}
+				
+			// Reject degenerate triangles.
+			if (VertexIndices[0] == VertexIndices[1] ||
+				VertexIndices[1] == VertexIndices[2] ||
+				VertexIndices[0] == VertexIndices[2])
+			{
+				continue;
+			}
 
-				{
-					FVector Edge01 = CornerPositions[1] - CornerPositions[0];
-					FVector Edge12 = CornerPositions[2] - CornerPositions[1];
-					FVector Edge20 = CornerPositions[0] - CornerPositions[2];
+			{
+				FVector Edge01 = CornerPositions[1] - CornerPositions[0];
+				FVector Edge12 = CornerPositions[2] - CornerPositions[1];
+				FVector Edge20 = CornerPositions[0] - CornerPositions[2];
 
-					float TriArea = 0.5f * ( Edge01 ^ Edge20 ).Size();
-					SurfaceArea += TriArea;
-				}
+				float TriArea = 0.5f * ( Edge01 ^ Edge20 ).Size();
+				SurfaceArea += TriArea;
+			}
 
-				Indexes.Add(VertexIndices[0]);
-				Indexes.Add(VertexIndices[1]);
-				Indexes.Add(VertexIndices[2]);
+			Indexes.Add(VertexIndices[0]);
+			Indexes.Add(VertexIndices[1]);
+			Indexes.Add(VertexIndices[2]);
 
-				MaterialIndexes.Add( PolygonGroupID.GetValue() );
-			}
+			MaterialIndexes.Add( PolygonGroupID.GetValue() );
 		}
 
 		uint32 NumVerts = Verts.Num();
@@ -447,7 +441,7 @@ public:
 				OutPolygonGroupMaterialNames[PolygonGroupID] = InPolygonGroupMaterialNames[PolygonGroupID];
 			}
 
-			TVertexAttributesRef<FVector> OutVertexPositions = OutReducedMesh.VertexAttributes().GetAttributesRef<FVector>(MeshAttribute::Vertex::Position);
+			TVertexAttributesRef<FVector> OutVertexPositions = OutReducedMesh.GetVertexPositions();
 
 			//Fill the vertex array
 			for (int32 VertexIndex = 0; VertexIndex < (int32)NumVerts; ++VertexIndex)
@@ -459,11 +453,12 @@ public:
 
 			TMap<int32, FPolygonGroupID> PolygonGroupMapping;
 
-			TVertexInstanceAttributesRef<FVector> OutVertexNormals = OutReducedMesh.VertexInstanceAttributes().GetAttributesRef<FVector>(MeshAttribute::VertexInstance::Normal);
-			TVertexInstanceAttributesRef<FVector> OutVertexTangents = OutReducedMesh.VertexInstanceAttributes().GetAttributesRef<FVector>(MeshAttribute::VertexInstance::Tangent);
-			TVertexInstanceAttributesRef<float> OutVertexBinormalSigns = OutReducedMesh.VertexInstanceAttributes().GetAttributesRef<float>(MeshAttribute::VertexInstance::BinormalSign);
-			TVertexInstanceAttributesRef<FVector4> OutVertexColors = OutReducedMesh.VertexInstanceAttributes().GetAttributesRef<FVector4>(MeshAttribute::VertexInstance::Color);
-			TVertexInstanceAttributesRef<FVector2D> OutVertexUVs = OutReducedMesh.VertexInstanceAttributes().GetAttributesRef<FVector2D>(MeshAttribute::VertexInstance::TextureCoordinate);
+			FStaticMeshAttributes Attributes(OutReducedMesh);
+			TVertexInstanceAttributesRef<FVector> OutVertexNormals = Attributes.GetVertexInstanceNormals();
+			TVertexInstanceAttributesRef<FVector> OutVertexTangents = Attributes.GetVertexInstanceTangents();
+			TVertexInstanceAttributesRef<float> OutVertexBinormalSigns = Attributes.GetVertexInstanceBinormalSigns();
+			TVertexInstanceAttributesRef<FVector4> OutVertexColors = Attributes.GetVertexInstanceColors();
+			TVertexInstanceAttributesRef<FVector2D> OutVertexUVs = Attributes.GetVertexInstanceUVs();
 
 			//Specify the number of texture coords in this mesh description
 			OutVertexUVs.SetNumChannels(InMeshNumTexCoords);
@@ -549,7 +544,7 @@ public:
 			TArray<FPolygonGroupID> ToDeletePolygonGroupIDs;
 			for (const FPolygonGroupID& PolygonGroupID : OutReducedMesh.PolygonGroups().GetElementIDs())
 			{
-				if (OutReducedMesh.GetPolygonGroupPolygons(PolygonGroupID).Num() == 0)
+				if (OutReducedMesh.GetPolygonGroupPolygonIDs(PolygonGroupID).Num() == 0)
 				{
 					ToDeletePolygonGroupIDs.Add(PolygonGroupID);
 				}
diff --git a/Engine/Source/Runtime/Datasmith/DatasmithCore/Private/DatasmithUtils.cpp b/Engine/Source/Runtime/Datasmith/DatasmithCore/Private/DatasmithUtils.cpp
--- a/Engine/Source/Runtime/Datasmith/DatasmithCore/Private/DatasmithUtils.cpp
+++ b/Engine/Source/Runtime/Datasmith/DatasmithCore/Private/DatasmithUtils.cpp
@@ -285,10 +285,11 @@ bool FDatasmithMeshUtils::ToMeshDescription(FDatasmithMesh& DsMesh, FMeshDescrip
 {
 	MeshDescription.Empty();
 
-	TVertexAttributesRef<FVector> VertexPositions = MeshDescription.VertexAttributes().GetAttributesRef<FVector>(MeshAttribute::Vertex::Position);
-	TVertexInstanceAttributesRef<FVector> VertexInstanceNormals = MeshDescription.VertexInstanceAttributes().GetAttributesRef<FVector>(MeshAttribute::VertexInstance::Normal);
-	TVertexInstanceAttributesRef<FVector2D> VertexInstanceUVs = MeshDescription.VertexInstanceAttributes().GetAttributesRef<FVector2D>(MeshAttribute::VertexInstance::TextureCoordinate);
-	TPolygonGroupAttributesRef<FName> PolygonGroupImportedMaterialSlotNames = MeshDescription.PolygonGroupAttributes().GetAttributesRef<FName>(MeshAttribute::PolygonGroup::ImportedMaterialSlotName);
+	FStaticMeshAttributes Attributes(MeshDescription);
+	TVertexAttributesRef<FVector> VertexPositions = Attributes.GetVertexPositions();
+	TVertexInstanceAttributesRef<FVector> VertexInstanceNormals = Attributes.GetVertexInstanceNormals();
+	TVertexInstanceAttributesRef<FVector2D> VertexInstanceUVs = Attributes.GetVertexInstanceUVs();
+	TPolygonGroupAttributesRef<FName> PolygonGroupImportedMaterialSlotNames = Attributes.GetPolygonGroupMaterialSlotNames();
 
 	// Prepared for static mesh usage ?
 	if (!ensure(VertexPositions.IsValid())
diff --git a/Engine/Source/Runtime/Engine/Private/StaticMesh.cpp b/Engine/Source/Runtime/Engine/Private/StaticMesh.cpp
--- a/Engine/Source/Runtime/Engine/Private/StaticMesh.cpp
+++ b/Engine/Source/Runtime/Engine/Private/StaticMesh.cpp
@@ -5300,21 +5300,18 @@ void UStaticMesh::BuildFromMeshDescription(const FMeshDescription& MeshDescripti
 		uint32 MinVertexIndex = TNumericLimits<uint32>::Max();
 		uint32 MaxVertexIndex = TNumericLimits<uint32>::Min();
 
-		for (FPolygonID PolygonID : MeshDescription.GetPolygonGroupPolygons(PolygonGroupID))
+		for (FTriangleID TriangleID : MeshDescription.GetPolygonGroupTriangles(PolygonGroupID))
 		{
-			for (FTriangleID TriangleID : MeshDescription.GetPolygonTriangleIDs(PolygonID))
+			for (FVertexInstanceID TriangleVertexInstanceIDs : MeshDescription.GetTriangleVertexInstances(TriangleID))
 			{
-				for (FVertexInstanceID TriangleVertexInstanceIDs : MeshDescription.GetTriangleVertexInstances(TriangleID))
-				{
-					uint32 VertexIndex = static_cast<uint32>(TriangleVertexInstanceIDs.GetValue());
-					MinVertexIndex = FMath::Min(MinVertexIndex, VertexIndex);
-					MaxVertexIndex = FMath::Max(MaxVertexIndex, VertexIndex);
-					IndexBuffer[IndexBufferIndex] = VertexIndex;
-					IndexBufferIndex++;
-				}
-
-				TriangleCount++;
+				uint32 VertexIndex = static_cast<uint32>(TriangleVertexInstanceIDs.GetValue());
+				MinVertexIndex = FMath::Min(MinVertexIndex, VertexIndex);
+				MaxVertexIndex = FMath::Max(MaxVertexIndex, VertexIndex);
+				IndexBuffer[IndexBufferIndex] = VertexIndex;
+				IndexBufferIndex++;
 			}
+
+			TriangleCount++;
 		}
 
 		Section.NumTriangles = TriangleCount;
@@ -5346,7 +5343,7 @@ void UStaticMesh::BuildFromMeshDescription(const FMeshDescription& MeshDescripti
 	for (uint32& Index : DepthOnlyIndexBuffer)
 	{
 		// Compress all vertex instances into the same instance for each vertex
-		Index = MeshDescription.GetVertexVertexInstances(MeshDescription.GetVertexInstanceVertex(FVertexInstanceID(Index)))[0].GetValue();
+		Index = MeshDescription.GetVertexVertexInstanceIDs(MeshDescription.GetVertexInstanceVertex(FVertexInstanceID(Index)))[0].GetValue();
 	}
 
 	LODResources.bHasDepthOnlyIndices = true;
@@ -6262,7 +6259,7 @@ void UStaticMesh::EnforceLightmapRestrictions(bool bUseRenderData)
 			{
 				if (const FMeshDescription* MeshDescription = GetMeshDescription(SourceLOD))
 				{
-					const TVertexInstanceAttributesConstRef<FVector2D> UVChannels = MeshDescription->VertexInstanceAttributes().GetAttributesRef<FVector2D>(MeshAttribute::VertexInstance::TextureCoordinate);
+					TVertexInstanceAttributesConstRef<FVector2D> UVChannels = FStaticMeshConstAttributes(*MeshDescription).GetVertexInstanceUVs();
 
 					// skip empty/stripped LODs
 					if (UVChannels.GetNumElements() > 0)
diff --git a/Engine/Source/Runtime/Landscape/Private/LandscapeEdit.cpp b/Engine/Source/Runtime/Landscape/Private/LandscapeEdit.cpp
--- a/Engine/Source/Runtime/Landscape/Private/LandscapeEdit.cpp
+++ b/Engine/Source/Runtime/Landscape/Private/LandscapeEdit.cpp
@@ -3118,14 +3118,15 @@ bool ALandscapeProxy::ExportToRawMesh(int32 InExportLOD, FMeshDescription& OutRa
 	const FIntRect LandscapeSectionRect = GetBoundingRect();
 	const FVector2D LandscapeUVScale = FVector2D(1.0f, 1.0f) / FVector2D(LandscapeSectionRect.Size());
 
-	TVertexAttributesRef<FVector> VertexPositions = OutRawMesh.VertexAttributes().GetAttributesRef<FVector>(MeshAttribute::Vertex::Position);
-	TEdgeAttributesRef<bool> EdgeHardnesses = OutRawMesh.EdgeAttributes().GetAttributesRef<bool>(MeshAttribute::Edge::IsHard);
-	TPolygonGroupAttributesRef<FName> PolygonGroupImportedMaterialSlotNames = OutRawMesh.PolygonGroupAttributes().GetAttributesRef<FName>(MeshAttribute::PolygonGroup::ImportedMaterialSlotName);
-	TVertexInstanceAttributesRef<FVector> VertexInstanceNormals = OutRawMesh.VertexInstanceAttributes().GetAttributesRef<FVector>(MeshAttribute::VertexInstance::Normal);
-	TVertexInstanceAttributesRef<FVector> VertexInstanceTangents = OutRawMesh.VertexInstanceAttributes().GetAttributesRef<FVector>(MeshAttribute::VertexInstance::Tangent);
-	TVertexInstanceAttributesRef<float> VertexInstanceBinormalSigns = OutRawMesh.VertexInstanceAttributes().GetAttributesRef<float>(MeshAttribute::VertexInstance::BinormalSign);
-	TVertexInstanceAttributesRef<FVector4> VertexInstanceColors = OutRawMesh.VertexInstanceAttributes().GetAttributesRef<FVector4>(MeshAttribute::VertexInstance::Color);
-	TVertexInstanceAttributesRef<FVector2D> VertexInstanceUVs = OutRawMesh.VertexInstanceAttributes().GetAttributesRef<FVector2D>(MeshAttribute::VertexInstance::TextureCoordinate);
+	FStaticMeshAttributes Attributes(OutRawMesh);
+	TVertexAttributesRef<FVector> VertexPositions = Attributes.GetVertexPositions();
+	TEdgeAttributesRef<bool> EdgeHardnesses = Attributes.GetEdgeHardnesses();
+	TPolygonGroupAttributesRef<FName> PolygonGroupImportedMaterialSlotNames = Attributes.GetPolygonGroupMaterialSlotNames();
+	TVertexInstanceAttributesRef<FVector> VertexInstanceNormals = Attributes.GetVertexInstanceNormals();
+	TVertexInstanceAttributesRef<FVector> VertexInstanceTangents = Attributes.GetVertexInstanceTangents();
+	TVertexInstanceAttributesRef<float> VertexInstanceBinormalSigns = Attributes.GetVertexInstanceBinormalSigns();
+	TVertexInstanceAttributesRef<FVector4> VertexInstanceColors = Attributes.GetVertexInstanceColors();
+	TVertexInstanceAttributesRef<FVector2D> VertexInstanceUVs = Attributes.GetVertexInstanceUVs();
 
 	if (VertexInstanceUVs.GetNumChannels() < 2)
 	{
diff --git a/Engine/Source/Runtime/MeshDescription/Private/MeshDescription.cpp b/Engine/Source/Runtime/MeshDescription/Private/MeshDescription.cpp
--- a/Engine/Source/Runtime/MeshDescription/Private/MeshDescription.cpp
+++ b/Engine/Source/Runtime/MeshDescription/Private/MeshDescription.cpp
@@ -56,31 +56,31 @@ void FMeshDescription::Initialize()
 	// Now register topology-based attributes
 
 	// Register vertex reference for vertex instances
-	VertexInstanceVertices = VertexInstanceElements->Get().GetAttributes().RegisterIndexAttribute<FVertexID>(MeshAttribute::VertexInstance::VertexIndex);
+	VertexInstanceVertices = VertexInstanceElements->Get().GetAttributes().RegisterIndexAttribute<FVertexID>(MeshAttribute::VertexInstance::VertexIndex, 1, EMeshAttributeFlags::Mandatory);
 
 	// Register vertex references for edges
-	EdgeVertices = EdgeElements->Get().GetAttributes().RegisterIndexAttribute<FVertexID[2]>(MeshAttribute::Edge::VertexIndex);
+	EdgeVertices = EdgeElements->Get().GetAttributes().RegisterIndexAttribute<FVertexID[2]>(MeshAttribute::Edge::VertexIndex, 1, EMeshAttributeFlags::Mandatory);
 
 	// Register vertex instance and polygon references for triangles
-	TriangleVertexInstances = TriangleElements->Get().GetAttributes().RegisterIndexAttribute<FVertexInstanceID[3]>(MeshAttribute::Triangle::VertexInstanceIndex);
-	TrianglePolygons = TriangleElements->Get().GetAttributes().RegisterIndexAttribute<FPolygonID>(MeshAttribute::Triangle::PolygonIndex);
-	TrianglePolygonGroups = TriangleElements->Get().GetAttributes().RegisterIndexAttribute<FPolygonGroupID>(MeshAttribute::Triangle::PolygonGroupIndex);
+	TriangleVertexInstances = TriangleElements->Get().GetAttributes().RegisterIndexAttribute<FVertexInstanceID[3]>(MeshAttribute::Triangle::VertexInstanceIndex, 1, EMeshAttributeFlags::Mandatory);
+	TrianglePolygons = TriangleElements->Get().GetAttributes().RegisterIndexAttribute<FPolygonID>(MeshAttribute::Triangle::PolygonIndex, 1, EMeshAttributeFlags::Mandatory);
+	TrianglePolygonGroups = TriangleElements->Get().GetAttributes().RegisterIndexAttribute<FPolygonGroupID>(MeshAttribute::Triangle::PolygonGroupIndex, 1, EMeshAttributeFlags::Mandatory);
 
 	// Register UV references for triangles
-	TriangleUVs = TriangleElements->Get().GetAttributes().RegisterIndexAttribute<FUVID[3]>(MeshAttribute::Triangle::UVIndex, 0);
+	TriangleUVs = TriangleElements->Get().GetAttributes().RegisterIndexAttribute<FUVID[3]>(MeshAttribute::Triangle::UVIndex, 0, EMeshAttributeFlags::Mandatory);
 
 	// Register vertex and edge references for triangles; these are transient references which are generated at load time
-	TriangleVertices = TriangleElements->Get().GetAttributes().RegisterIndexAttribute<FVertexID[3]>(MeshAttribute::Triangle::VertexIndex, 1, EMeshAttributeFlags::IndexReference | EMeshAttributeFlags::Transient);
-	TriangleEdges = TriangleElements->Get().GetAttributes().RegisterIndexAttribute<FEdgeID[3]>(MeshAttribute::Triangle::EdgeIndex, 1, EMeshAttributeFlags::IndexReference | EMeshAttributeFlags::Transient);
+	TriangleVertices = TriangleElements->Get().GetAttributes().RegisterIndexAttribute<FVertexID[3]>(MeshAttribute::Triangle::VertexIndex, 1, EMeshAttributeFlags::Mandatory | EMeshAttributeFlags::Transient);
+	TriangleEdges = TriangleElements->Get().GetAttributes().RegisterIndexAttribute<FEdgeID[3]>(MeshAttribute::Triangle::EdgeIndex, 1, EMeshAttributeFlags::Mandatory | EMeshAttributeFlags::Transient);
 
 	// Register polygon group reference for polygons
-	PolygonPolygonGroups = PolygonElements->Get().GetAttributes().RegisterIndexAttribute<FPolygonGroupID>(MeshAttribute::Polygon::PolygonGroupIndex);
+	PolygonPolygonGroups = PolygonElements->Get().GetAttributes().RegisterIndexAttribute<FPolygonGroupID>(MeshAttribute::Polygon::PolygonGroupIndex, 1, EMeshAttributeFlags::Mandatory);
 
 	// Minimal requirement is that vertices have a Position attribute
-	VertexPositions = VertexElements->Get().GetAttributes().RegisterAttribute(MeshAttribute::Vertex::Position, 1, FVector::ZeroVector, EMeshAttributeFlags::Lerpable);
+	VertexPositions = VertexElements->Get().GetAttributes().RegisterAttribute(MeshAttribute::Vertex::Position, 1, FVector::ZeroVector, EMeshAttributeFlags::Lerpable | EMeshAttributeFlags::Mandatory);
 
 	// Register UVCoordinates attribute for UVs
-	UVElements->Get().GetAttributes().RegisterAttribute(MeshAttribute::UV::UVCoordinate, 1, FVector2D::ZeroVector, EMeshAttributeFlags::Lerpable);
+	UVElements->Get().GetAttributes().RegisterAttribute(MeshAttribute::UV::UVCoordinate, 1, FVector2D::ZeroVector, EMeshAttributeFlags::Lerpable | EMeshAttributeFlags::Mandatory);
 
 	// Associate indexers with element types and their referencing attributes
 	InitializeIndexers();
@@ -134,8 +134,8 @@ void FMeshDescription::Cache()
 	PolygonGroupElements = Elements.Find(PolygonGroupsName)->Get();
 
 	// Register required transient attributes
-	TriangleVertices = TriangleElements->Get().GetAttributes().RegisterIndexAttribute<FVertexID[3]>(MeshAttribute::Triangle::VertexIndex, 1, EMeshAttributeFlags::IndexReference | EMeshAttributeFlags::Transient);
-	TriangleEdges = TriangleElements->Get().GetAttributes().RegisterIndexAttribute<FEdgeID[3]>(MeshAttribute::Triangle::EdgeIndex, 1, EMeshAttributeFlags::IndexReference | EMeshAttributeFlags::Transient);
+	TriangleVertices = TriangleElements->Get().GetAttributes().RegisterIndexAttribute<FVertexID[3]>(MeshAttribute::Triangle::VertexIndex, 1, EMeshAttributeFlags::Mandatory | EMeshAttributeFlags::Transient);
+	TriangleEdges = TriangleElements->Get().GetAttributes().RegisterIndexAttribute<FEdgeID[3]>(MeshAttribute::Triangle::EdgeIndex, 1, EMeshAttributeFlags::Mandatory | EMeshAttributeFlags::Transient);
 
 	// Cache fundamental attribute arrays
 	VertexInstanceVertices = VertexInstanceElements->Get().GetAttributes().GetAttributesRef<FVertexID>(MeshAttribute::VertexInstance::VertexIndex);
@@ -752,80 +752,72 @@ void FMeshDescription::DeleteTriangle_Internal(const FTriangleID TriangleID, TCo
 	const FPolygonID PolygonID = TrianglePolygons[TriangleID];
 	PolygonToTriangles.RemoveReferenceFromKey(PolygonID, TriangleID);
 
-	if (PolygonToTriangles.Find(PolygonID).Num() == 0)
-	{
-		// If it was the only triangle in the polygon, delete the polygon too
-		TArrayView<FVertexInstanceID> TriVertexInstanceIDs = TriangleVertexInstances[TriangleID];
-		TArrayView<FEdgeID> TriEdgeIDs = TriangleEdges[TriangleID];
+	TArrayView<FVertexInstanceID> TriVertexInstanceIDs = TriangleVertexInstances[TriangleID];
+	TArrayView<FEdgeID> TriEdgeIDs = TriangleEdges[TriangleID];
 
-		for (int32 Index = 0; Index < 3; ++Index)
-		{
-			// Remove vertex instance from the triangle. Here's how we do that safely:
-			// 1) Take a copy of the VertexInstance ID
-			// 2) Remove reference from the VertexInstance to Triangles indexer
-			// 3) Set the triangle VertexInstance to INDEX_NONE
-			// 4) Ask the indexer if that VertexInstance ID has any triangle references, and if not, add it to the list.
-			//
-			// Step 3 is important because Find() below might need to reindex if the key is stale, and this would cause it to
-			// re-add the triangle VertexInstance reference if it were still valid, since the triangle hasn't yet been deleted.
-			//
-			// We could also wait until the triangle were deleted before querying the indexer, but this would involve taking
-			// copies of all the referenced element IDs first.
+	for (int32 Index = 0; Index < 3; ++Index)
+	{
+		// Remove vertex instance from the triangle. Here's how we do that safely:
+		// 1) Take a copy of the VertexInstance ID
+		// 2) Remove reference from the VertexInstance to Triangles indexer
+		// 3) Set the triangle VertexInstance to INDEX_NONE
+		// 4) Ask the indexer if that VertexInstance ID has any triangle references, and if not, add it to the list.
+		//
+		// Step 3 is important because Find() below might need to reindex if the key is stale, and this would cause it to
+		// re-add the triangle VertexInstance reference if it were still valid, since the triangle hasn't yet been deleted.
+		//
+		// We could also wait until the triangle were deleted before querying the indexer, but this would involve taking
+		// copies of all the referenced element IDs first.
 
-			const FVertexInstanceID VertexInstanceID = TriVertexInstanceIDs[Index];
+		const FVertexInstanceID VertexInstanceID = TriVertexInstanceIDs[Index];
 			
-			VertexInstanceToTriangles.RemoveReferenceFromKey(VertexInstanceID, TriangleID);
-			TriVertexInstanceIDs[Index] = INDEX_NONE;
+		VertexInstanceToTriangles.RemoveReferenceFromKey(VertexInstanceID, TriangleID);
+		TriVertexInstanceIDs[Index] = INDEX_NONE;
 			
-			if (InOutOrphanedVertexInstancesPtr && VertexInstanceToTriangles.Find(VertexInstanceID).Num() == 0)
-			{
-				AddUnique(*InOutOrphanedVertexInstancesPtr, VertexInstanceID);
-			}
+		if (InOutOrphanedVertexInstancesPtr && VertexInstanceToTriangles.Find(VertexInstanceID).Num() == 0)
+		{
+			AddUnique(*InOutOrphanedVertexInstancesPtr, VertexInstanceID);
+		}
 
-			// Remove edge from the triangle. Same rules as above apply!
+		// Remove edge from the triangle. Same rules as above apply!
 			
-			const FEdgeID EdgeID = TriEdgeIDs[Index];
+		const FEdgeID EdgeID = TriEdgeIDs[Index];
 			
-			EdgeToTriangles.RemoveReferenceFromKey(EdgeID, TriangleID);
-			TriEdgeIDs[Index] = INDEX_NONE;
+		EdgeToTriangles.RemoveReferenceFromKey(EdgeID, TriangleID);
+		TriEdgeIDs[Index] = INDEX_NONE;
 			
-			if (InOutOrphanedEdgesPtr && EdgeToTriangles.Find(EdgeID).Num() == 0)
-			{
-				AddUnique(*InOutOrphanedEdgesPtr, EdgeID);
-			}
+		if (InOutOrphanedEdgesPtr && EdgeToTriangles.Find(EdgeID).Num() == 0)
+		{
+			AddUnique(*InOutOrphanedEdgesPtr, EdgeID);
+		}
 
-			// Remove UVs from the triangle if there are any
+		// Remove UVs from the triangle if there are any
 
-			for (int32 UVIndex = 0; UVIndex < TriangleUVs.GetNumChannels(); UVIndex++)
-			{
-				TArrayView<const FUVID> TriUVs = TriangleUVs.Get(TriangleID, UVIndex);
-				FUVID UVID = TriUVs[Index];
-				UVToTriangles.RemoveReferenceFromKey(UVID, TriangleID, UVIndex);
-			}
+		for (int32 UVIndex = 0; UVIndex < TriangleUVs.GetNumChannels(); UVIndex++)
+		{
+			TArrayView<const FUVID> TriUVs = TriangleUVs.Get(TriangleID, UVIndex);
+			FUVID UVID = TriUVs[Index];
+			UVToTriangles.RemoveReferenceFromKey(UVID, TriangleID, UVIndex);
 		}
+	}
 
-		// Remove the polygon
+	if (PolygonToTriangles.Find(PolygonID).Num() == 0)
+	{
+		// If it was the only triangle in the polygon, delete the polygon too
 		check(PolygonPolygonGroups[PolygonID] == PolygonGroupID);
 		PolygonGroupToPolygons.RemoveReferenceFromKey(PolygonGroupID, PolygonID);
 
 		PolygonElements->Get().Remove(PolygonID);
 		PolygonToTriangles.RemoveKey(PolygonID);
-
-		// Check orphaned polygon groups after the polygon has been removed so that reindexing won't find it again.
-		if (InOutOrphanedPolygonGroupsPtr && PolygonGroupToPolygons.Find(PolygonGroupID).Num() == 0)
-		{
-			AddUnique(*InOutOrphanedPolygonGroupsPtr, PolygonGroupID);
-		}
-	}
-	else
-	{
-		// @todo: Handle this properly when deleting a triangle which forms part of an n-gon
-		// Either it needs to shave off the triangle from the contour and update the contour vertex instances,
-		// or it should just refuse to delete the triangle.
-		check(false);
 	}
 
 	TriangleElements->Get().Remove(TriangleID);
+
+	// Check orphaned polygon groups after the triangle has been removed so that reindexing won't find it again.
+	if (InOutOrphanedPolygonGroupsPtr && PolygonGroupToTriangles.Find(PolygonGroupID).Num() == 0)
+	{
+		AddUnique(*InOutOrphanedPolygonGroupsPtr, PolygonGroupID);
+	}
 }
 
 void FMeshDescription::DeleteTriangle(const FTriangleID TriangleID, TArray<FEdgeID>* InOutOrphanedEdgesPtr, TArray<FVertexInstanceID>* InOutOrphanedVertexInstancesPtr, TArray<FPolygonGroupID>* InOutOrphanedPolygonGroupsPtr)
diff --git a/Engine/Source/Runtime/MeshDescription/Private/MeshDescriptionBase.cpp b/Engine/Source/Runtime/MeshDescription/Private/MeshDescriptionBase.cpp
--- a/Engine/Source/Runtime/MeshDescription/Private/MeshDescriptionBase.cpp
+++ b/Engine/Source/Runtime/MeshDescription/Private/MeshDescriptionBase.cpp
@@ -362,7 +362,7 @@ void UMeshDescriptionBase::GetVertexConnectedEdges(FVertexID VertexID, TArray<FE
 		return;
 	}
 
-	OutEdgeIDs = MeshDescription.GetVertexConnectedEdges(VertexID);
+	OutEdgeIDs = MeshDescription.GetVertexConnectedEdgeIDs(VertexID);
 }
 
 int32 UMeshDescriptionBase::GetNumVertexConnectedEdges(FVertexID VertexID) const
@@ -384,7 +384,7 @@ void UMeshDescriptionBase::GetVertexVertexInstances(FVertexID VertexID, TArray<F
 		return;
 	}
 
-	OutVertexInstanceIDs = MeshDescription.GetVertexVertexInstances(VertexID);
+	OutVertexInstanceIDs = MeshDescription.GetVertexVertexInstanceIDs(VertexID);
 }
 
 int32 UMeshDescriptionBase::GetNumVertexVertexInstances(FVertexID VertexID) const
@@ -527,7 +527,7 @@ void UMeshDescriptionBase::GetVertexInstanceConnectedTriangles(FVertexInstanceID
 		return;
 	}
 
-	OutConnectedTriangleIDs = MeshDescription.GetVertexInstanceConnectedTriangles(VertexInstanceID);
+	OutConnectedTriangleIDs = MeshDescription.GetVertexInstanceConnectedTriangleIDs(VertexInstanceID);
 }
 
 int32 UMeshDescriptionBase::GetNumVertexInstanceConnectedTriangles(FVertexInstanceID VertexInstanceID) const
@@ -603,7 +603,7 @@ void UMeshDescriptionBase::GetEdgeConnectedTriangles(FEdgeID EdgeID, TArray<FTri
 		return;
 	}
 
-	OutConnectedTriangleIDs = MeshDescription.GetEdgeConnectedTriangles(EdgeID);
+	OutConnectedTriangleIDs = MeshDescription.GetEdgeConnectedTriangleIDs(EdgeID);
 }
 
 int32 UMeshDescriptionBase::GetNumEdgeConnectedTriangles(FEdgeID EdgeID) const
@@ -790,7 +790,7 @@ void UMeshDescriptionBase::GetPolygonTriangles(FPolygonID PolygonID, TArray<FTri
 		return;
 	}
 
-	OutTriangleIDs = MeshDescription.GetPolygonTriangleIDs(PolygonID);
+	OutTriangleIDs = MeshDescription.GetPolygonTriangles(PolygonID);
 }
 
 int32 UMeshDescriptionBase::GetNumPolygonTriangles(FPolygonID PolygonID) const
@@ -974,7 +974,7 @@ void UMeshDescriptionBase::GetPolygonGroupPolygons(FPolygonGroupID PolygonGroupI
 		return;
 	}
 
-	OutPolygonIDs = MeshDescription.GetPolygonGroupPolygons(PolygonGroupID);
+	OutPolygonIDs = MeshDescription.GetPolygonGroupPolygonIDs(PolygonGroupID);
 }
 
 int32 UMeshDescriptionBase::GetNumPolygonGroupPolygons(FPolygonGroupID PolygonGroupID) const
diff --git a/Engine/Source/Runtime/MeshDescription/Public/MeshAttributeArray.h b/Engine/Source/Runtime/MeshDescription/Public/MeshAttributeArray.h
--- a/Engine/Source/Runtime/MeshDescription/Public/MeshAttributeArray.h
+++ b/Engine/Source/Runtime/MeshDescription/Public/MeshAttributeArray.h
@@ -25,7 +25,7 @@ using AttributeTypes = TTuple
 	FVector,
 	FVector2D,
 	float,
-	int,
+	int32,
 	bool,
 	FName
 >;
@@ -270,7 +270,8 @@ enum class EMeshAttributeFlags : uint32
 	AutoGenerated		= (1 << 1),		/** Attribute is auto-generated by importer or editable mesh, rather than representing an imported property */	
 	Mergeable			= (1 << 2),		/** If all vertices' attributes are mergeable, and of near-equal value, they can be welded */
 	Transient			= (1 << 3),		/** Attribute is not serialized */
-	IndexReference		= (1 << 4)		/** Attribute is a reference to another element index */
+	IndexReference		= (1 << 4),		/** Attribute is a reference to another element index */
+	Mandatory			= (1 << 5),		/** Attribute is required in the mesh description */
 };
 
 ENUM_CLASS_FLAGS(EMeshAttributeFlags);
@@ -572,7 +573,19 @@ public:
 
 	/** Implicitly construct a TMeshAttributesRef-to-const from a regular one */
 	template <typename T = AttributeType, typename TEnableIf<TIsSame<T, const T>::Value, int>::Type = 0>
-	FORCEINLINE TMeshAttributesRef(TMeshAttributesRef<ElementIDType, typename TRemoveCV<T>::Type> InRef)
+	FORCEINLINE TMeshAttributesRef(TMeshAttributesRef<ElementIDType, typename TRemoveCV<T>::Type, false> InRef)
+		: ArrayPtr(InRef.ArrayPtr)
+	{}
+
+	/** Implicitly construct a TMeshAttributesRef from a TMeshAttributesArray **/
+	template <typename T = ElementIDType, typename U = AttributeType, typename TEnableIf<!TIsSame<T, int32>::Value, int>::Type = 0>
+	FORCEINLINE TMeshAttributesRef(TMeshAttributesRef<int32, U, false> InRef)
+		: ArrayPtr(InRef.ArrayPtr)
+	{}
+
+	/** Implicitly construct a TMeshAttributesRef-to-const from a TMeshAttributesArray */
+	template <typename T = ElementIDType, typename U = AttributeType, typename TEnableIf<TIsSame<U, const U>::Value, int>::Type = 0, typename TEnableIf<!TIsSame<T, int32>::Value, int>::Type = 0>
+	FORCEINLINE TMeshAttributesRef(TMeshAttributesRef<int32, typename TRemoveCV<U>::Type, false> InRef)
 		: ArrayPtr(InRef.ArrayPtr)
 	{}
 
@@ -750,11 +763,26 @@ public:
 
 	/** Implicitly construct a TMeshAttributesRef-to-const from a regular one */
 	template <typename T = Type, typename TEnableIf<TIsSame<T, const T>::Value, int>::Type = 0>
-	FORCEINLINE TMeshAttributesRef(TMeshAttributesRef<ElementIDType, TArrayView<typename TRemoveCV<T>::Type>> InRef)
+	FORCEINLINE TMeshAttributesRef(TMeshAttributesRef<ElementIDType, TArrayView<typename TRemoveCV<T>::Type>, true> InRef)
+		: ArrayPtr(InRef.ArrayPtr),
+		  Extent(InRef.Extent)
+	{}
+
+	/** Implicitly construct a TMeshAttributesRef from a TMeshAttributesArray **/
+	template <typename T = ElementIDType, typename U = AttributeType, typename TEnableIf<!TIsSame<T, int32>::Value, int>::Type = 0>
+	FORCEINLINE TMeshAttributesRef(TMeshAttributesRef<int32, U, true> InRef)
+		: ArrayPtr(InRef.ArrayPtr),
+		  Extent(InRef.Extent)
+	{}
+
+	/** Implicitly construct a TMeshAttributesRef-to-const from a TMeshAttributesArray */
+	template <typename T = ElementIDType, typename U = Type, typename TEnableIf<TIsSame<U, const U>::Value, int>::Type = 0, typename TEnableIf<!TIsSame<T, int32>::Value, int>::Type = 0>
+	FORCEINLINE TMeshAttributesRef(TMeshAttributesRef<int32, TArrayView<typename TRemoveCV<U>::Type>, true> InRef)
 		: ArrayPtr(InRef.ArrayPtr),
 		  Extent(InRef.Extent)
 	{}
 
+
 	/** Access elements from attribute channel 0 */
 	template <typename T = ElementIDType, typename TEnableIf<TIsDerivedFrom<T, FElementID>::Value, int>::Type = 0>
 	FORCEINLINE TArrayView<Type> operator[](const ElementIDType ElementID) const
@@ -1104,7 +1132,7 @@ public:
 	 */
 	template <typename AttributeType,
 			  typename TEnableIf<TIsSame<AttributeType, int>::Value || TIsDerivedFrom<AttributeType, FElementID>::Value, int>::Type = 0>
-	TMeshAttributesArray<AttributeType> RegisterIndexAttribute(const FName AttributeName, const int32 NumberOfChannels = 1, const EMeshAttributeFlags Flags = EMeshAttributeFlags::IndexReference);
+	TMeshAttributesArray<AttributeType> RegisterIndexAttribute(const FName AttributeName, const int32 NumberOfChannels = 1, const EMeshAttributeFlags Flags = EMeshAttributeFlags::None);
 
 	/**
 	 * Register a new attribute denoting an array of indices (type is implicitly int).
@@ -1115,7 +1143,7 @@ public:
 			  typename AttributeType = typename TBreakArray<ArrayType>::Type,
 			  typename TEnableIf<TIsSame<AttributeType, int>::Value || TIsDerivedFrom<AttributeType, FElementID>::Value, int>::Type = 0,
 			  uint32 Extent = TBreakArray<ArrayType>::Value>
-	TMeshAttributesArray<TArrayView<AttributeType>> RegisterIndexAttribute(const FName AttributeName, const int32 NumberOfChannels = 1, const EMeshAttributeFlags Flags = EMeshAttributeFlags::IndexReference);
+	TMeshAttributesArray<TArrayView<AttributeType>> RegisterIndexAttribute(const FName AttributeName, const int32 NumberOfChannels = 1, const EMeshAttributeFlags Flags = EMeshAttributeFlags::None);
 
 	/**
 	 * Unregister an attribute with the given name.
@@ -1180,6 +1208,28 @@ public:
 		Map.GetKeys(OutAttributeNames);
 	}
 
+	/** Determine whether an attribute has any of the given flags */
+	bool DoesAttributeHaveAnyFlags(const FName AttributeName, EMeshAttributeFlags AttributeFlags) const
+	{
+		if (const FAttributesSetEntry* ArraySetPtr = Map.Find(AttributeName))
+		{
+			return EnumHasAnyFlags((*ArraySetPtr)->GetFlags(), AttributeFlags);
+		}
+
+		return false;
+	}
+
+	/** Determine whether an attribute has all of the given flags */
+	bool DoesAttributeHaveAllFlags(const FName AttributeName, EMeshAttributeFlags AttributeFlags) const
+	{
+		if (const FAttributesSetEntry* ArraySetPtr = Map.Find(AttributeName))
+		{
+			return EnumHasAllFlags((*ArraySetPtr)->GetFlags(), AttributeFlags);
+		}
+
+		return false;
+	}
+
 	uint32 GetHash(const FName AttributeName) const
 	{
 		if (const FAttributesSetEntry* ArraySetPtr = Map.Find(AttributeName))
@@ -1377,13 +1427,13 @@ TMeshAttributesArray<AttributeType> FAttributesSetBase::RegisterIndexAttribute(c
 		else
 		{
 			Map.Remove(AttributeName);
-			FAttributesSetEntry& Entry = Map.Emplace(AttributeName, FAttributesSetEntry(NumberOfChannels, int32(INDEX_NONE), Flags, NumElements, 1));
+			FAttributesSetEntry& Entry = Map.Emplace(AttributeName, FAttributesSetEntry(NumberOfChannels, int32(INDEX_NONE), Flags | EMeshAttributeFlags::IndexReference, NumElements, 1));
 			return TMeshAttributesArray<AttributeType>(static_cast<ArrayType*>(Entry.Get()));
 		}
 	}
 	else
 	{
-		FAttributesSetEntry& Entry = Map.Emplace(AttributeName, FAttributesSetEntry(NumberOfChannels, int32(INDEX_NONE), Flags, NumElements, 1));
+		FAttributesSetEntry& Entry = Map.Emplace(AttributeName, FAttributesSetEntry(NumberOfChannels, int32(INDEX_NONE), Flags | EMeshAttributeFlags::IndexReference, NumElements, 1));
 		return TMeshAttributesArray<AttributeType>(static_cast<ArrayType*>(Entry.Get()));
 	}
 
@@ -1408,13 +1458,13 @@ TMeshAttributesArray<TArrayView<AttributeType>> FAttributesSetBase::RegisterInde
 		else
 		{
 			Map.Remove(AttributeName);
-			FAttributesSetEntry& Entry = Map.Emplace(AttributeName, FAttributesSetEntry(NumberOfChannels, int32(INDEX_NONE), Flags, NumElements, Extent));
+			FAttributesSetEntry& Entry = Map.Emplace(AttributeName, FAttributesSetEntry(NumberOfChannels, int32(INDEX_NONE), Flags | EMeshAttributeFlags::IndexReference, NumElements, Extent));
 			return TMeshAttributesArray<TArrayView<AttributeType>>(static_cast<TMeshAttributeArraySet<AttributeType>*>(Entry.Get()), Extent);
 		}
 	}
 	else
 	{
-		FAttributesSetEntry& Entry = Map.Emplace(AttributeName, FAttributesSetEntry(NumberOfChannels, int32(INDEX_NONE), Flags, NumElements, Extent));
+		FAttributesSetEntry& Entry = Map.Emplace(AttributeName, FAttributesSetEntry(NumberOfChannels, int32(INDEX_NONE), Flags | EMeshAttributeFlags::IndexReference, NumElements, Extent));
 		return TMeshAttributesArray<TArrayView<AttributeType>>(static_cast<TMeshAttributeArraySet<AttributeType>*>(Entry.Get()), Extent);
 	}
 }
@@ -1453,7 +1503,7 @@ public:
 	{
 		if (const FAttributesSetEntry* ArraySetPtr = this->Map.Find(AttributeName))
 		{
-			using Type = typename TChooseClass<TIsDerivedFrom<AttributeType, FElementID>::Value, int, AttributeType>::Result;
+			using Type = typename TChooseClass<TIsDerivedFrom<AttributeType, FElementID>::Value, int32, AttributeType>::Result;
 			if ((*ArraySetPtr)->HasType<Type>() && (*ArraySetPtr)->GetExtent() == 1)
 			{
 				return TMeshAttributesConstRef<ElementIDType, AttributeType>(static_cast<const TMeshAttributeArraySet<AttributeType>*>(ArraySetPtr->Get()));
@@ -1469,7 +1519,7 @@ public:
 	{
 		if (FAttributesSetEntry* ArraySetPtr = this->Map.Find(AttributeName))
 		{
-			using Type = typename TChooseClass<TIsDerivedFrom<AttributeType, FElementID>::Value, int, AttributeType>::Result;
+			using Type = typename TChooseClass<TIsDerivedFrom<AttributeType, FElementID>::Value, int32, AttributeType>::Result;
 			if ((*ArraySetPtr)->HasType<Type>() && (*ArraySetPtr)->GetExtent() == 1)
 			{
 				return TMeshAttributesRef<ElementIDType, AttributeType>(static_cast<TMeshAttributeArraySet<AttributeType>*>(ArraySetPtr->Get()));
@@ -1486,7 +1536,7 @@ public:
 	{
 		if (const FAttributesSetEntry* ArraySetPtr = this->Map.Find(AttributeName))
 		{
-			using Type = typename TChooseClass<TIsDerivedFrom<AttributeType, FElementID>::Value, int, AttributeType>::Result;
+			using Type = typename TChooseClass<TIsDerivedFrom<AttributeType, FElementID>::Value, int32, AttributeType>::Result;
 			if ((*ArraySetPtr)->HasType<Type>())
 			{
 				return TMeshAttributesConstRef<ElementIDType, ArrayType>(static_cast<const TMeshAttributeArraySet<AttributeType>*>(ArraySetPtr->Get()), (*ArraySetPtr)->GetExtent());
@@ -1503,7 +1553,7 @@ public:
 	{
 		if (FAttributesSetEntry* ArraySetPtr = this->Map.Find(AttributeName))
 		{
-			using Type = typename TChooseClass<TIsDerivedFrom<AttributeType, FElementID>::Value, int, AttributeType>::Result;
+			using Type = typename TChooseClass<TIsDerivedFrom<AttributeType, FElementID>::Value, int32, AttributeType>::Result;
 			if ((*ArraySetPtr)->HasType<Type>())
 			{
 				return TMeshAttributesRef<ElementIDType, ArrayType>(static_cast<TMeshAttributeArraySet<AttributeType>*>(ArraySetPtr->Get()), (*ArraySetPtr)->GetExtent());
diff --git a/Engine/Source/Runtime/MeshDescription/Public/MeshAttributes.h b/Engine/Source/Runtime/MeshDescription/Public/MeshAttributes.h
--- a/Engine/Source/Runtime/MeshDescription/Public/MeshAttributes.h
+++ b/Engine/Source/Runtime/MeshDescription/Public/MeshAttributes.h
@@ -132,6 +132,17 @@ public:
 		return MeshDescription.TriangleAttributes().GetAttributesRef<TArrayView<FUVID>>(MeshAttribute::Triangle::UVIndex);
 	}
 
+	/** Accessors for UV coordinates */
+	TUVAttributesRef<FVector2D> GetUVCoordinates(int32 UVChannel)
+	{
+		return MeshDescription.UVAttributes(UVChannel).GetAttributesRef<FVector2D>(MeshAttribute::UV::UVCoordinate);
+	}
+
+	TUVAttributesConstRef<FVector2D> GetUVCoordinates(int32 UVChannel) const
+	{
+		return MeshDescription.UVAttributes(UVChannel).GetAttributesRef<FVector2D>(MeshAttribute::UV::UVCoordinate);
+	}
+
 	/** Accessors for array of polygon group IDs for triangles */
 	TTriangleAttributesRef<FPolygonGroupID> GetTrianglePolygonGroupIndices()
 	{
@@ -206,6 +217,12 @@ public:
 		return MeshDescription.TriangleAttributes().GetAttributesRef<TArrayView<FUVID>>(MeshAttribute::Triangle::UVIndex);
 	}
 
+	/** Accessors for UV coordinates */
+	TUVAttributesConstRef<FVector2D> GetUVCoordinates(int32 UVChannel) const
+	{
+		return MeshDescription.UVAttributes(UVChannel).GetAttributesRef<FVector2D>(MeshAttribute::UV::UVCoordinate);
+	}
+
 	/** Accessors for array of polygon group IDs for triangles */
 	TTriangleAttributesConstRef<FPolygonGroupID> GetTrianglePolygonGroupIndices() const
 	{
diff --git a/Engine/Source/Runtime/MeshDescription/Public/MeshDescription.h b/Engine/Source/Runtime/MeshDescription/Public/MeshDescription.h
--- a/Engine/Source/Runtime/MeshDescription/Public/MeshDescription.h
+++ b/Engine/Source/Runtime/MeshDescription/Public/MeshDescription.h
@@ -429,11 +429,17 @@ public:
 	FEdgeID GetVertexPairEdge(const FVertexID VertexID0, const FVertexID VertexID1) const;
 
 	/** Returns reference to an array of Edge IDs connected to this vertex */
-	TArrayView<const FEdgeID> GetVertexConnectedEdges(const FVertexID VertexID) const
+	TArrayView<const FEdgeID> GetVertexConnectedEdgeIDs(const FVertexID VertexID) const
 	{
 		return VertexToEdges.Find<FEdgeID>(VertexID);
 	}
 
+	UE_DEPRECATED(4.26, "Please use GetVertexConnectedEdgeIDs instead.")
+	TArray<FEdgeID> GetVertexConnectedEdges(const FVertexID VertexID) const
+	{
+		return TArray<FEdgeID>(GetVertexConnectedEdgeIDs(VertexID));
+	}
+
 	/** Returns number of edges connected to this vertex */
 	int32 GetNumVertexConnectedEdges(const FVertexID VertexID) const
 	{
@@ -441,11 +447,17 @@ public:
 	}
 
 	/** Returns reference to an array of VertexInstance IDs instanced from this vertex */
-	TArrayView<const FVertexInstanceID> GetVertexVertexInstances(const FVertexID VertexID) const
+	TArrayView<const FVertexInstanceID> GetVertexVertexInstanceIDs(const FVertexID VertexID) const
 	{
 		return VertexToVertexInstances.Find<FVertexInstanceID>(VertexID);
 	}
 
+	UE_DEPRECATED(4.26, "Please use GetVertexVertexInstanceIDs instead.")
+	TArray<FVertexInstanceID> GetVertexVertexInstances(const FVertexID VertexID) const
+	{
+		return TArray<FVertexInstanceID>(GetVertexVertexInstanceIDs(VertexID));
+	}
+
 	/** Returns number of vertex instances created from this vertex */
 	int32 GetNumVertexVertexInstances(const FVertexID VertexID) const
 	{
@@ -566,7 +578,12 @@ public:
 		return VertexPositions[VertexID];
 	}
 
-	TMeshAttributesArray<FVector> GetVertexPositions() const
+	TVertexAttributesRef<FVector> GetVertexPositions()
+	{
+		return VertexPositions;
+	}
+
+	TVertexAttributesRef<const FVector> GetVertexPositions() const
 	{
 		return VertexPositions;
 	}
@@ -584,11 +601,17 @@ public:
 	FEdgeID GetVertexInstancePairEdge(const FVertexInstanceID VertexInstanceID0, const FVertexInstanceID VertexInstanceID1) const;
 
 	/** Returns reference to an array of Triangle IDs connected to this vertex instance */
-	TArrayView<const FTriangleID> GetVertexInstanceConnectedTriangles(const FVertexInstanceID VertexInstanceID) const
+	TArrayView<const FTriangleID> GetVertexInstanceConnectedTriangleIDs(const FVertexInstanceID VertexInstanceID) const
 	{
 		return VertexInstanceToTriangles.Find<FTriangleID>(VertexInstanceID);
 	}
 
+	UE_DEPRECATED(4.26, "Please use GetVertexInstanceTriangleIDs() instead.")
+	TArray<FTriangleID> GetVertexInstanceConnectedTriangles(const FVertexInstanceID VertexInstanceID) const
+	{
+		return TArray<FTriangleID>(GetVertexInstanceConnectedTriangleIDs(VertexInstanceID));
+	}
+
 	/** Returns the number of triangles connected to this vertex instance */
 	int32 GetNumVertexInstanceConnectedTriangles(const FVertexInstanceID VertexInstanceID) const
 	{
@@ -651,11 +674,17 @@ public:
 	}
 
 	/** Returns reference to an array of triangle IDs connected to this edge */
-	TArrayView<const FTriangleID> GetEdgeConnectedTriangles(const FEdgeID EdgeID) const
+	TArrayView<const FTriangleID> GetEdgeConnectedTriangleIDs(const FEdgeID EdgeID) const
 	{
 		return EdgeToTriangles.Find<FTriangleID>(EdgeID);
 	}
 
+	UE_DEPRECATED(4.26, "Please use GetEdgeConnectedTriangleIDs() instead.")
+	TArray<FTriangleID> GetEdgeConnectedTriangles(const FEdgeID EdgeID) const
+	{
+		return TArray<FTriangleID>(GetEdgeConnectedTriangleIDs(EdgeID));
+	}
+
 	int32 GetNumEdgeConnectedTriangles(const FEdgeID EdgeID) const
 	{
 		return EdgeToTriangles.Find(EdgeID).Num();
@@ -841,11 +870,17 @@ public:
 	// Polygon operations
 
 	/** Return reference to an array of triangle IDs which comprise this polygon */
-	TArrayView<const FTriangleID> GetPolygonTriangleIDs(const FPolygonID PolygonID) const
+	TArrayView<const FTriangleID> GetPolygonTriangles(const FPolygonID PolygonID) const
 	{
 		return PolygonToTriangles.Find<FTriangleID>(PolygonID);
 	}
 
+	UE_DEPRECATED(4.26, "Please use GetPolygonTriangles() instead.")
+	TArray<FTriangleID> GetPolygonTriangleIDs(const FPolygonID PolygonID) const
+	{
+		return TArray<FTriangleID>(GetPolygonTriangles(PolygonID));
+	}
+
 	/** Return the number of triangles which comprise this polygon */
 	int32 GetNumPolygonTriangles(const FPolygonID PolygonID) const
 	{
@@ -991,7 +1026,7 @@ public:
 			TArray<FVertexInstanceID> VertexInstanceIDs = GetPolygonVertexInstances(PolygonID);
 			for (const FVertexInstanceID VertexInstanceID : VertexInstanceIDs)
 			{
-				for (const FEdgeID EdgeID : GetVertexConnectedEdges(GetVertexInstanceVertex(VertexInstanceID)))
+				for (const FEdgeID EdgeID : GetVertexConnectedEdgeIDs(GetVertexInstanceVertex(VertexInstanceID)))
 				{
 					if (!OutEdgeIDs.Contains(EdgeID) && IsEdgeInternalToPolygon(EdgeID, PolygonID))
 					{
@@ -1113,11 +1148,17 @@ public:
 	// Polygon group operations
 
 	/** Returns the polygons associated with the given polygon group */
-	TArrayView<const FPolygonID> GetPolygonGroupPolygons(const FPolygonGroupID PolygonGroupID) const
+	TArrayView<const FPolygonID> GetPolygonGroupPolygonIDs(const FPolygonGroupID PolygonGroupID) const
 	{
 		return PolygonGroupToPolygons.Find<FPolygonID>(PolygonGroupID);
 	}
 
+	UE_DEPRECATED(4.26, "Please use GetPolygonGroupPolygonIDs() instead.")
+	TArray<FPolygonID> GetPolygonGroupPolygons(const FPolygonGroupID PolygonGroupID) const
+	{
+		return TArray<FPolygonID>(GetPolygonGroupPolygonIDs(PolygonGroupID));
+	}
+
 	/** Returns the triangles associated with the given polygon group */
 	TArrayView<const FTriangleID> GetPolygonGroupTriangles(const FPolygonGroupID PolygonGroupID) const
 	{
diff --git a/Engine/Source/Runtime/StaticMeshDescription/Private/StaticMeshAttributes.cpp b/Engine/Source/Runtime/StaticMeshDescription/Private/StaticMeshAttributes.cpp
--- a/Engine/Source/Runtime/StaticMeshDescription/Private/StaticMeshAttributes.cpp
+++ b/Engine/Source/Runtime/StaticMeshDescription/Private/StaticMeshAttributes.cpp
@@ -11,7 +11,6 @@ namespace MeshAttribute
 	const FName VertexInstance::Tangent("Tangent");
 	const FName VertexInstance::BinormalSign("BinormalSign");
 	const FName VertexInstance::Color("Color");
-	const FName VertexInstance::UVIndex("UVIndex");
 
 	const FName Edge::IsHard("IsHard");
 	const FName Edge::IsUVSeam("IsUVSeam");
@@ -38,19 +37,19 @@ void FStaticMeshAttributes::Register()
 	// Add basic vertex attributes
 
 	// Add basic vertex instance attributes
-	MeshDescription.VertexInstanceAttributes().RegisterAttribute<FVector2D>(MeshAttribute::VertexInstance::TextureCoordinate, 1, FVector2D::ZeroVector, EMeshAttributeFlags::Lerpable);
-	MeshDescription.VertexInstanceAttributes().RegisterAttribute<FVector>(MeshAttribute::VertexInstance::Normal, 1, FVector::ZeroVector, EMeshAttributeFlags::AutoGenerated);
-	MeshDescription.VertexInstanceAttributes().RegisterAttribute<FVector>(MeshAttribute::VertexInstance::Tangent, 1, FVector::ZeroVector, EMeshAttributeFlags::AutoGenerated);
-	MeshDescription.VertexInstanceAttributes().RegisterAttribute<float>(MeshAttribute::VertexInstance::BinormalSign, 1, 0.0f, EMeshAttributeFlags::AutoGenerated);
-	MeshDescription.VertexInstanceAttributes().RegisterAttribute<FVector4>(MeshAttribute::VertexInstance::Color, 1, FVector4(1.0f, 1.0f, 1.0f, 1.0f), EMeshAttributeFlags::Lerpable);
+	MeshDescription.VertexInstanceAttributes().RegisterAttribute<FVector2D>(MeshAttribute::VertexInstance::TextureCoordinate, 1, FVector2D::ZeroVector, EMeshAttributeFlags::Lerpable | EMeshAttributeFlags::Mandatory);
+	MeshDescription.VertexInstanceAttributes().RegisterAttribute<FVector>(MeshAttribute::VertexInstance::Normal, 1, FVector::ZeroVector, EMeshAttributeFlags::AutoGenerated | EMeshAttributeFlags::Mandatory);
+	MeshDescription.VertexInstanceAttributes().RegisterAttribute<FVector>(MeshAttribute::VertexInstance::Tangent, 1, FVector::ZeroVector, EMeshAttributeFlags::AutoGenerated | EMeshAttributeFlags::Mandatory);
+	MeshDescription.VertexInstanceAttributes().RegisterAttribute<float>(MeshAttribute::VertexInstance::BinormalSign, 1, 0.0f, EMeshAttributeFlags::AutoGenerated | EMeshAttributeFlags::Mandatory);
+	MeshDescription.VertexInstanceAttributes().RegisterAttribute<FVector4>(MeshAttribute::VertexInstance::Color, 1, FVector4(1.0f, 1.0f, 1.0f, 1.0f), EMeshAttributeFlags::Lerpable | EMeshAttributeFlags::Mandatory);
 
 	// Add basic edge attributes
-	MeshDescription.EdgeAttributes().RegisterAttribute<bool>(MeshAttribute::Edge::IsHard, 1, false);
+	MeshDescription.EdgeAttributes().RegisterAttribute<bool>(MeshAttribute::Edge::IsHard, 1, false, EMeshAttributeFlags::Mandatory);
 
 	// Add basic polygon attributes
 
 	// Add basic polygon group attributes
-	MeshDescription.PolygonGroupAttributes().RegisterAttribute<FName>(MeshAttribute::PolygonGroup::ImportedMaterialSlotName); //The unique key to match the mesh material slot
+	MeshDescription.PolygonGroupAttributes().RegisterAttribute<FName>(MeshAttribute::PolygonGroup::ImportedMaterialSlotName, 1, NAME_None, EMeshAttributeFlags::Mandatory); //The unique key to match the mesh material slot
 
 	// Call super class
 	FMeshAttributes::Register();
diff --git a/Engine/Source/Runtime/StaticMeshDescription/Private/StaticMeshOperations.cpp b/Engine/Source/Runtime/StaticMeshDescription/Private/StaticMeshOperations.cpp
--- a/Engine/Source/Runtime/StaticMeshDescription/Private/StaticMeshOperations.cpp
+++ b/Engine/Source/Runtime/StaticMeshDescription/Private/StaticMeshOperations.cpp
@@ -47,7 +47,7 @@ static bool GetPolygonTangentsAndNormals(FMeshDescription& MeshDescription,
 	Center /= float(VertexInstanceIDs.Num());
 
 	float AdjustedComparisonThreshold = FMath::Max(ComparisonThreshold, MIN_flt);
-	for (const FTriangleID TriangleID : MeshDescription.GetPolygonTriangleIDs(PolygonID))
+	for (const FTriangleID TriangleID : MeshDescription.GetPolygonTriangles(PolygonID))
 	{
 		TArrayView<const FVertexInstanceID> TriangleVertexInstances = MeshDescription.GetTriangleVertexInstances(TriangleID);
 		const FVertexID VertexID0 = MeshDescription.GetVertexInstanceVertex(TriangleVertexInstances[0]);
@@ -336,7 +336,7 @@ void FStaticMeshOperations::DetermineEdgeHardnessesFromVertexInstanceNormals(FMe
 			UniqueVertexInstanceIDs.Reset(ReservedElements);
 
 			// Get a list of all vertex instances for this vertex which form part of any polygon connected to the edge
-			for (const FVertexInstanceID VertexInstanceID : MeshDescription.GetVertexVertexInstances(VertexID))
+			for (const FVertexInstanceID VertexInstanceID : MeshDescription.GetVertexVertexInstanceIDs(VertexID))
 			{
 				for (const FPolygonID PolygonID : MeshDescription.GetVertexInstanceConnectedPolygons<TInlineAllocator<8>>(VertexInstanceID))
 				{
@@ -433,7 +433,7 @@ void FStaticMeshOperations::ConvertHardEdgesToSmoothGroup(const FMeshDescription
 
 	for (const FPolygonID PolygonID : SourceMeshDescription.Polygons().GetElementIDs())
 	{
-		TriangleCount += SourceMeshDescription.GetPolygonTriangleIDs(PolygonID).Num();
+		TriangleCount += SourceMeshDescription.GetPolygonTriangles(PolygonID).Num();
 		if (ConsumedPolygons[PolygonID.GetValue()])
 		{
 			continue;
@@ -539,7 +539,7 @@ void FStaticMeshOperations::ConvertHardEdgesToSmoothGroup(const FMeshDescription
 	for (const FPolygonID PolygonID : SourceMeshDescription.Polygons().GetElementIDs())
 	{
 		uint32 PolygonSmoothValue = PolygonSmoothGroup[PolygonID];
-		for (const FTriangleID TriangleID : SourceMeshDescription.GetPolygonTriangleIDs(PolygonID))
+		for (const FTriangleID TriangleID : SourceMeshDescription.GetPolygonTriangles(PolygonID))
 		{
 			FaceSmoothingMasks[TriangleIndex++] = PolygonSmoothValue;
 		}
@@ -608,15 +608,14 @@ void FStaticMeshOperations::ConvertToRawMesh(const FMeshDescription& SourceMeshD
 	DestinationRawMesh.Empty();
 
 	//Gather all array data
-	TVertexAttributesConstRef<FVector> VertexPositions = SourceMeshDescription.VertexAttributes().GetAttributesRef<FVector>(MeshAttribute::Vertex::Position);
-
-	TVertexInstanceAttributesConstRef<FVector> VertexInstanceNormals = SourceMeshDescription.VertexInstanceAttributes().GetAttributesRef<FVector>(MeshAttribute::VertexInstance::Normal);
-	TVertexInstanceAttributesConstRef<FVector> VertexInstanceTangents = SourceMeshDescription.VertexInstanceAttributes().GetAttributesRef<FVector>(MeshAttribute::VertexInstance::Tangent);
-	TVertexInstanceAttributesConstRef<float> VertexInstanceBinormalSigns = SourceMeshDescription.VertexInstanceAttributes().GetAttributesRef<float>(MeshAttribute::VertexInstance::BinormalSign);
-	TVertexInstanceAttributesConstRef<FVector4> VertexInstanceColors = SourceMeshDescription.VertexInstanceAttributes().GetAttributesRef<FVector4>(MeshAttribute::VertexInstance::Color);
-	TVertexInstanceAttributesConstRef<FVector2D> VertexInstanceUVs = SourceMeshDescription.VertexInstanceAttributes().GetAttributesRef<FVector2D>(MeshAttribute::VertexInstance::TextureCoordinate);
-
-	TPolygonGroupAttributesConstRef<FName> PolygonGroupMaterialSlotName = SourceMeshDescription.PolygonGroupAttributes().GetAttributesRef<FName>(MeshAttribute::PolygonGroup::ImportedMaterialSlotName);
+	FStaticMeshConstAttributes Attributes(SourceMeshDescription);
+	TVertexAttributesConstRef<FVector> VertexPositions = Attributes.GetVertexPositions();
+	TVertexInstanceAttributesConstRef<FVector> VertexInstanceNormals = Attributes.GetVertexInstanceNormals();
+	TVertexInstanceAttributesConstRef<FVector> VertexInstanceTangents = Attributes.GetVertexInstanceTangents();
+	TVertexInstanceAttributesConstRef<float> VertexInstanceBinormalSigns = Attributes.GetVertexInstanceBinormalSigns();
+	TVertexInstanceAttributesConstRef<FVector4> VertexInstanceColors = Attributes.GetVertexInstanceColors();
+	TVertexInstanceAttributesConstRef<FVector2D> VertexInstanceUVs = Attributes.GetVertexInstanceUVs();
+	TPolygonGroupAttributesConstRef<FName> PolygonGroupMaterialSlotName = Attributes.GetPolygonGroupMaterialSlotNames();
 
 	DestinationRawMesh.VertexPositions.AddZeroed(SourceMeshDescription.Vertices().Num());
 	TArray<int32> RemapVerts;
@@ -656,7 +655,7 @@ void FStaticMeshOperations::ConvertToRawMesh(const FMeshDescription& SourceMeshD
 	{
 		const FPolygonGroupID& PolygonGroupID = SourceMeshDescription.GetPolygonPolygonGroup(PolygonID);
 		int32 PolygonIDValue = PolygonID.GetValue();
-		TArrayView<const FTriangleID> TriangleIDs = SourceMeshDescription.GetPolygonTriangleIDs(PolygonID);
+		TArrayView<const FTriangleID> TriangleIDs = SourceMeshDescription.GetPolygonTriangles(PolygonID);
 		for (const FTriangleID TriangleID : TriangleIDs)
 		{
 			if (MaterialMap.Num() > 0 && MaterialMap.Contains(PolygonGroupMaterialSlotName[PolygonGroupID]))
@@ -697,7 +696,7 @@ void FStaticMeshOperations::ConvertToRawMesh(const FMeshDescription& SourceMeshD
 //We will also weld the vertex position (old FRawMesh is not always welded) and construct a mapping array to match the FVertexID
 void FillMeshDescriptionVertexPositionNoDuplicate(const TArray<FVector>& RawMeshVertexPositions, FMeshDescription& DestinationMeshDescription, TArray<FVertexID>& RemapVertexPosition)
 {
-	TVertexAttributesRef<FVector> VertexPositions = DestinationMeshDescription.VertexAttributes().GetAttributesRef<FVector>(MeshAttribute::Vertex::Position);
+	TVertexAttributesRef<FVector> VertexPositions = DestinationMeshDescription.GetVertexPositions();
 
 	const int32 NumVertex = RawMeshVertexPositions.Num();
 
@@ -1271,7 +1270,7 @@ void FStaticMeshOperations::ComputeTangentsAndNormals(FMeshDescription& MeshDesc
 	const int32 NumTriangles = MeshDescription.Triangles().Num();
 	check(MeshDescription.Triangles().GetArraySize() == NumTriangles);
 
-	//vCompute the weight (area and angle) for each triangles
+	// Compute the weight (area and angle) for each triangles
 	TArray<FTriangleData> TriangleDatas;
 	TriangleDatas.SetNum(NumTriangles);
 	if (bComputeWeightedNormals)
@@ -1349,9 +1348,9 @@ void FStaticMeshOperations::ComputeTangentsAndNormals(FMeshDescription& MeshDesc
 
 				bool bPointHasAllTangents = true;
 				// Fill the VertexInfoMap
-				for (const FEdgeID EdgeID : MeshDescription.GetVertexConnectedEdges(VertexID))
+				for (const FEdgeID EdgeID : MeshDescription.GetVertexConnectedEdgeIDs(VertexID))
 				{
-					for (const FTriangleID TriangleID : MeshDescription.GetEdgeConnectedTriangles(EdgeID))
+					for (const FTriangleID TriangleID : MeshDescription.GetEdgeConnectedTriangleIDs(EdgeID))
 					{
 						FVertexInfo& VertexInfo = VertexInfoMap.FindOrAdd(TriangleID);
 						int32 EdgeIndex = VertexInfo.EdgeIDs.AddUnique(EdgeID);
@@ -1416,7 +1415,7 @@ void FStaticMeshOperations::ComputeTangentsAndNormals(FMeshDescription& MeshDesc
 								//End of the group
 								continue;
 							}
-							for (const FTriangleID TriangleID : MeshDescription.GetEdgeConnectedTriangles(EdgeID))
+							for (const FTriangleID TriangleID : MeshDescription.GetEdgeConnectedTriangleIDs(EdgeID))
 							{
 								if (TriangleID == CurrentVertexInfo.TriangleID)
 								{
@@ -1681,12 +1680,12 @@ void FStaticMeshOperations::FindOverlappingCorners(FOverlappingCorners& OutOverl
 	TArray<MeshDescriptionOperationNamespace::FIndexAndZ> VertIndexAndZ;
 	VertIndexAndZ.Reserve(NumWedges);
 
-	TVertexAttributesConstRef<FVector> VertexPositions = MeshDescription.VertexAttributes().GetAttributesRef<FVector>(MeshAttribute::Vertex::Position);
+	TVertexAttributesConstRef<FVector> VertexPositions = MeshDescription.GetVertexPositions();
 
 	int32 WedgeIndex = 0;
 	for (const FPolygonID PolygonID : MeshDescription.Polygons().GetElementIDs())
 	{
-		TArrayView<const FTriangleID> TriangleIDs = MeshDescription.GetPolygonTriangleIDs(PolygonID);
+		TArrayView<const FTriangleID> TriangleIDs = MeshDescription.GetPolygonTriangles(PolygonID);
 		for (const FTriangleID TriangleID : TriangleIDs)
 		{
 			for (int32 Corner = 0; Corner < 3; ++Corner)
@@ -1740,12 +1739,14 @@ struct FLayoutUVMeshDescriptionView final : FLayoutUV::IMeshView
 
 	FLayoutUVMeshDescriptionView(FMeshDescription& InMeshDescription, uint32 InSrcChannel, uint32 InDstChannel)
 		: MeshDescription(InMeshDescription)
-		, Positions(InMeshDescription.VertexAttributes().GetAttributesRef<FVector>(MeshAttribute::Vertex::Position))
-		, Normals(InMeshDescription.VertexInstanceAttributes().GetAttributesRef<FVector>(MeshAttribute::VertexInstance::Normal))
-		, TexCoords(InMeshDescription.VertexInstanceAttributes().GetAttributesRef<FVector2D>(MeshAttribute::VertexInstance::TextureCoordinate))
 		, SrcChannel(InSrcChannel)
 		, DstChannel(InDstChannel)
 	{
+		FStaticMeshAttributes Attributes(InMeshDescription);
+		Positions = Attributes.GetVertexPositions();
+		Normals = Attributes.GetVertexInstanceNormals();
+		TexCoords = Attributes.GetVertexInstanceUVs();
+
 		uint32 NumTris = MeshDescription.Triangles().Num();
 
 		NumIndices = NumTris * 3;
@@ -1757,7 +1758,7 @@ struct FLayoutUVMeshDescriptionView final : FLayoutUV::IMeshView
 
 		for (const FPolygonID PolygonID : MeshDescription.Polygons().GetElementIDs())
 		{
-			TArrayView<const FTriangleID> TriangleIDs = MeshDescription.GetPolygonTriangleIDs(PolygonID);
+			TArrayView<const FTriangleID> TriangleIDs = MeshDescription.GetPolygonTriangles(PolygonID);
 			for (const FTriangleID TriangleID : TriangleIDs)
 			{
 				for (int32 Corner = 0; Corner < 3; ++Corner)
@@ -2013,7 +2014,7 @@ void FStaticMeshOperations::GeneratePlanarUV(const FMeshDescription& MeshDescrip
 	FVector U = FVector::UpVector;
 	FVector V = FVector::RightVector;
 
-	TMeshAttributesConstRef<FVertexID, FVector> VertexPositions = MeshDescription.VertexAttributes().GetAttributesRef<FVector>(MeshAttribute::Vertex::Position);
+	TMeshAttributesConstRef<FVertexID, FVector> VertexPositions = MeshDescription.GetVertexPositions();
 
 	OutTexCoords.Reserve(MeshDescription.VertexInstances().Num());
 
@@ -2045,7 +2046,7 @@ void FStaticMeshOperations::GenerateCylindricalUV(FMeshDescription& MeshDescript
 	FVector V = FVector::ForwardVector;
 	Offset.X -= Size.X / 2.f;
 
-	TMeshAttributesConstRef<FVertexID, FVector> VertexPositions = MeshDescription.VertexAttributes().GetAttributesRef<FVector>(MeshAttribute::Vertex::Position);
+	TMeshAttributesConstRef<FVertexID, FVector> VertexPositions = MeshDescription.GetVertexPositions();
 
 	OutTexCoords.Reserve(MeshDescription.VertexInstances().Num());
 
@@ -2120,7 +2121,7 @@ void FStaticMeshOperations::GenerateBoxUV(const FMeshDescription& MeshDescriptio
 	FVector Size = Params.Size * Params.Scale;
 	FVector HalfSize = Size / 2.0f;
 
-	TMeshAttributesConstRef<FVertexID, FVector> VertexPositions = MeshDescription.VertexAttributes().GetAttributesRef<FVector>(MeshAttribute::Vertex::Position);
+	TMeshAttributesConstRef<FVertexID, FVector> VertexPositions = MeshDescription.GetVertexPositions();
 
 	OutTexCoords.Reserve(MeshDescription.VertexInstances().Num());
 
@@ -2218,7 +2219,7 @@ void FStaticMeshOperations::SwapPolygonPolygonGroup(FMeshDescription& MeshDescri
 
 	for (const FPolygonID PolygonID : MeshDescription.Polygons().GetElementIDs())
 	{
-		int32 TriangleCount = MeshDescription.GetPolygonTriangleIDs(PolygonID).Num();
+		int32 TriangleCount = MeshDescription.GetPolygonTriangles(PolygonID).Num();
 		if (TriangleIndex >= TriangleIndexStart && TriangleIndex < TriangleIndexEnd)
 		{
 			check(TriangleIndex + (TriangleCount - 1) < TriangleIndexEnd);
@@ -2226,7 +2227,7 @@ void FStaticMeshOperations::SwapPolygonPolygonGroup(FMeshDescription& MeshDescri
 			if (OldpolygonGroupID != TargetPolygonGroupID)
 			{
 				MeshDescription.SetPolygonPolygonGroup(PolygonID, TargetPolygonGroupID);
-				if (bRemoveEmptyPolygonGroup && MeshDescription.GetPolygonGroupPolygons(OldpolygonGroupID).Num() < 1)
+				if (bRemoveEmptyPolygonGroup && MeshDescription.GetPolygonGroupPolygonIDs(OldpolygonGroupID).Num() < 1)
 				{
 					MeshDescription.DeletePolygonGroup(OldpolygonGroupID);
 				}
@@ -2254,7 +2255,7 @@ bool FStaticMeshOperations::HasVertexColor(const FMeshDescription& MeshDescripti
 
 void FStaticMeshOperations::BuildWeldedVertexIDRemap(const FMeshDescription& MeshDescription, const float WeldingThreshold, TMap<FVertexID, FVertexID>& OutVertexIDRemap)
 {
-	TVertexAttributesConstRef<FVector> VertexPositions = MeshDescription.VertexAttributes().GetAttributesRef<FVector>(MeshAttribute::Vertex::Position);
+	TVertexAttributesConstRef<FVector> VertexPositions = MeshDescription.GetVertexPositions();
 
 	int32 NumVertex = MeshDescription.Vertices().Num();
 	OutVertexIDRemap.Reserve(NumVertex);
diff --git a/Engine/Source/Runtime/StaticMeshDescription/Public/StaticMeshAttributes.h b/Engine/Source/Runtime/StaticMeshDescription/Public/StaticMeshAttributes.h
--- a/Engine/Source/Runtime/StaticMeshDescription/Public/StaticMeshAttributes.h
+++ b/Engine/Source/Runtime/StaticMeshDescription/Public/StaticMeshAttributes.h
@@ -20,7 +20,6 @@ namespace MeshAttribute
 		extern STATICMESHDESCRIPTION_API const FName Tangent;
 		extern STATICMESHDESCRIPTION_API const FName BinormalSign;
 		extern STATICMESHDESCRIPTION_API const FName Color;
-		extern STATICMESHDESCRIPTION_API const FName UVIndex;
 	}
 
 	namespace Edge
@@ -156,7 +155,6 @@ public:
 	TVertexInstanceAttributesConstRef<FVector> GetVertexInstanceTangents() const { return MeshDescription.VertexInstanceAttributes().GetAttributesRef<FVector>(MeshAttribute::VertexInstance::Tangent); }
 	TVertexInstanceAttributesConstRef<float> GetVertexInstanceBinormalSigns() const { return MeshDescription.VertexInstanceAttributes().GetAttributesRef<float>(MeshAttribute::VertexInstance::BinormalSign); }
 	TVertexInstanceAttributesConstRef<FVector4> GetVertexInstanceColors() const { return MeshDescription.VertexInstanceAttributes().GetAttributesRef<FVector4>(MeshAttribute::VertexInstance::Color); }
-	TVertexInstanceAttributesConstRef<int> GetVertexInstanceUVIndices() const { return MeshDescription.VertexInstanceAttributes().GetAttributesRef<int>(MeshAttribute::VertexInstance::UVIndex); }
 	TEdgeAttributesConstRef<bool> GetEdgeHardnesses() const { return MeshDescription.EdgeAttributes().GetAttributesRef<bool>(MeshAttribute::Edge::IsHard); }
 	UE_DEPRECATED(4.26, "This attribute is no longer supported, please remove code pertaining to it.")
 	TEdgeAttributesConstRef<float> GetEdgeCreaseSharpnesses() const { return MeshDescription.EdgeAttributes().GetAttributesRef<float>(MeshAttribute::Edge::CreaseSharpness); }
