Fix for module dependencies on Mac causing SkeletalMeshEditor load too early, causing a crash

diff --git a/Engine/Source/Developer/MeshMergeUtilities/MeshMergeUtilities.Build.cs b/Engine/Source/Developer/MeshMergeUtilities/MeshMergeUtilities.Build.cs
--- a/Engine/Source/Developer/MeshMergeUtilities/MeshMergeUtilities.Build.cs
+++ b/Engine/Source/Developer/MeshMergeUtilities/MeshMergeUtilities.Build.cs
@@ -22,7 +22,6 @@ public class MeshMergeUtilities : ModuleRules
                 "SlateCore",
                 "Slate",
                 "StaticMeshEditor",
-                "SkeletalMeshEditor",
                 "MaterialBaking",
                 "MeshUtilitiesCommon",
 				"ToolMenus",
@@ -42,6 +41,8 @@ public class MeshMergeUtilities : ModuleRules
                "HierarchicalLODUtilities",
                "MeshUtilities",
                "MeshReductionInterface",
+			   "SkeletalMeshEditor",
+			   "Persona",
           }
         );
 
