Added IMeshMergeUtilities::RetrieveMeshDescription
* General cleanup
* Removed ExtractImposterToRawMesh() which was overspecialized and not specifically handling Imposters
* Renamed remains of FRawMesh to MeshDescription

diff --git a/Engine/Source/Developer/HierarchicalLODUtilities/Private/HierarchicalLODUtilities.cpp b/Engine/Source/Developer/HierarchicalLODUtilities/Private/HierarchicalLODUtilities.cpp
--- a/Engine/Source/Developer/HierarchicalLODUtilities/Private/HierarchicalLODUtilities.cpp
+++ b/Engine/Source/Developer/HierarchicalLODUtilities/Private/HierarchicalLODUtilities.cpp
@@ -319,7 +319,7 @@ UStaticMesh* CreateImposterStaticMesh(UStaticMeshComponent* InComponent, UMateri
 
 		FMeshDescription* ImposterMesh = StaticMesh->CreateMeshDescription(0);
 		const IMeshMergeUtilities& MeshMergeUtilities = FModuleManager::Get().LoadModuleChecked<IMeshMergeModule>("MeshMergeUtilities").GetUtilities();
-		MeshMergeUtilities.ExtractImposterToRawMesh(InComponent, *ImposterMesh);
+		MeshMergeUtilities.RetrieveMeshDescription(InComponent->GetStaticMesh(), InComponent->GetStaticMesh()->GetNumLODs() - 1, *ImposterMesh);
 
 		// Disable collisions on imposters
 		FMeshSectionInfo Info = StaticMesh->GetSectionInfoMap().Get(0, 0);
diff --git a/Engine/Source/Developer/MeshMergeUtilities/Private/MeshMergeHelpers.cpp b/Engine/Source/Developer/MeshMergeUtilities/Private/MeshMergeHelpers.cpp
--- a/Engine/Source/Developer/MeshMergeUtilities/Private/MeshMergeHelpers.cpp
+++ b/Engine/Source/Developer/MeshMergeUtilities/Private/MeshMergeHelpers.cpp
@@ -181,27 +181,27 @@ void FMeshMergeHelpers::ExtractSections(const UStaticMesh* StaticMesh, int32 LOD
 	}
 }
 
-void FMeshMergeHelpers::ExpandInstances(const UInstancedStaticMeshComponent* InInstancedStaticMeshComponent, FMeshDescription& InOutRawMesh, TArray<FSectionInfo>& InOutSections)
+void FMeshMergeHelpers::ExpandInstances(const UInstancedStaticMeshComponent* InInstancedStaticMeshComponent, FMeshDescription& InOutMeshDescription, TArray<FSectionInfo>& InOutSections)
 {
 	TRACE_CPUPROFILER_EVENT_SCOPE(FMeshMergeHelpers::ExpandInstances)
 
-	FMeshDescription CombinedRawMesh;
-	FStaticMeshAttributes(CombinedRawMesh).Register();
+	FMeshDescription CombinedMeshDescription;
+	FStaticMeshAttributes(CombinedMeshDescription).Register();
 
 	FTransform ComponentTransform = InInstancedStaticMeshComponent->GetComponentTransform();
 	FTransform ComponentTransformInv = ComponentTransform.Inverse();
 
 	for(const FInstancedStaticMeshInstanceData& InstanceData : InInstancedStaticMeshComponent->PerInstanceSMData)
 	{
-		FMeshDescription InstanceRawMesh = InOutRawMesh;
-		FMeshMergeHelpers::TransformRawMeshVertexData(ComponentTransformInv * FTransform(InstanceData.Transform) * ComponentTransform, InstanceRawMesh);
-		FMeshMergeHelpers::AppendRawMesh(CombinedRawMesh, InstanceRawMesh);
+		FMeshDescription InstanceMeshDescription = InOutMeshDescription;
+		TransformMeshDescriptionVertexData(ComponentTransformInv * FTransform(InstanceData.Transform) * ComponentTransform, InstanceMeshDescription);
+		AppendMesh(CombinedMeshDescription, InstanceMeshDescription);
 	}
 
-	InOutRawMesh = CombinedRawMesh;
+	InOutMeshDescription = CombinedMeshDescription;
 }
 
-void FMeshMergeHelpers::RetrieveMesh(const UStaticMeshComponent* StaticMeshComponent, int32 LODIndex, FMeshDescription& RawMesh, bool bPropagateVertexColours)
+void FMeshMergeHelpers::RetrieveMesh(const UStaticMeshComponent* StaticMeshComponent, int32 LODIndex, FMeshDescription& OutMeshDescription, bool bPropagateVertexColours)
 {
 	TRACE_CPUPROFILER_EVENT_SCOPE(FMeshMergeHelpers::RetrieveMesh)
 
@@ -213,11 +213,11 @@ void FMeshMergeHelpers::RetrieveMesh(const UStaticMeshComponent* StaticMeshCompo
 	// Imported meshes will have a valid mesh description
 	const bool bImportedMesh = StaticMesh->IsMeshDescriptionValid(LODIndex);
 		
-	// Export the raw mesh data using static mesh render data
-	ExportStaticMeshLOD(StaticMesh->RenderData->LODResources[LODIndex], RawMesh, StaticMesh->StaticMaterials);
+	// Export the mesh data using static mesh render data
+	ExportStaticMeshLOD(StaticMesh->RenderData->LODResources[LODIndex], OutMeshDescription, StaticMesh->StaticMaterials);
 
-	// Make sure the raw mesh is not irreparably malformed.
-	if (RawMesh.VertexInstances().Num() <= 0)
+	// Make sure the mesh is not irreparably malformed.
+	if (OutMeshDescription.VertexInstances().Num() <= 0)
 	{
 		return;
 	}
@@ -225,27 +225,27 @@ void FMeshMergeHelpers::RetrieveMesh(const UStaticMeshComponent* StaticMeshCompo
 	// Use build settings from base mesh for LOD entries that was generated inside Editor.
 	const FMeshBuildSettings& BuildSettings = bImportedMesh ? StaticMeshModel.BuildSettings : StaticMesh->GetSourceModel(0).BuildSettings;
 
-	// Transform raw mesh to world space
+	// Transform mesh to world space
 	FTransform ComponentToWorldTransform = StaticMeshComponent->GetComponentTransform();
 
 	// Handle spline mesh deformation
 	if (bIsSplineMeshComponent)
 	{
 		const USplineMeshComponent* SplineMeshComponent = Cast<USplineMeshComponent>(StaticMeshComponent);
-		// Deform raw mesh data according to the Spline Mesh Component's data
-		PropagateSplineDeformationToRawMesh(SplineMeshComponent, RawMesh);
+		// Deform mesh data according to the Spline Mesh Component's data
+		PropagateSplineDeformationToMesh(SplineMeshComponent, OutMeshDescription);
 	}
 
 	// If specified propagate painted vertex colors into our raw mesh
 	if (bPropagateVertexColours)
 	{
-		PropagatePaintedColorsToRawMesh(StaticMeshComponent, LODIndex, RawMesh);
+		PropagatePaintedColorsToMesh(StaticMeshComponent, LODIndex, OutMeshDescription);
 	}
 
-	// Transform raw mesh vertex data by the Static Mesh Component's component to world transformation	
-	TransformRawMeshVertexData(ComponentToWorldTransform, RawMesh);
+	// Transform mesh vertex data by the Static Mesh Component's component to world transformation	
+	TransformMeshDescriptionVertexData(ComponentToWorldTransform, OutMeshDescription);
 
-	if (RawMesh.VertexInstances().Num() <= 0)
+	if (OutMeshDescription.VertexInstances().Num() <= 0)
 	{
 		return;
 	}
@@ -261,11 +261,11 @@ void FMeshMergeHelpers::RetrieveMesh(const UStaticMeshComponent* StaticMeshCompo
 	{
 		ComputeNTBsOptions |= EComputeNTBsFlags::UseMikkTSpace;
 	}
-	FStaticMeshOperations::ComputeTriangleTangentsAndNormals(RawMesh, 0.0f);
-	FStaticMeshOperations::RecomputeNormalsAndTangentsIfNeeded(RawMesh, ComputeNTBsOptions);
+	FStaticMeshOperations::ComputeTriangleTangentsAndNormals(OutMeshDescription, 0.0f);
+	FStaticMeshOperations::RecomputeNormalsAndTangentsIfNeeded(OutMeshDescription, ComputeNTBsOptions);
 }
 
-void FMeshMergeHelpers::RetrieveMesh(USkeletalMeshComponent* SkeletalMeshComponent, int32 LODIndex, FMeshDescription& RawMesh, bool bPropagateVertexColours)
+void FMeshMergeHelpers::RetrieveMesh(const USkeletalMeshComponent* SkeletalMeshComponent, int32 LODIndex, FMeshDescription& OutMeshDescription, bool bPropagateVertexColours)
 {
 	FSkeletalMeshModel* Resource = SkeletalMeshComponent->SkeletalMesh->GetImportedModel();
 	if (Resource->LODModels.IsValidIndex(LODIndex))
@@ -274,16 +274,18 @@ void FMeshMergeHelpers::RetrieveMesh(USkeletalMeshComponent* SkeletalMeshCompone
 
 		// Get the CPU skinned verts for this LOD
 		TArray<FFinalSkinVertex> FinalVertices;
+
+		// GetCPUSkinnedVertices 
 		SkeletalMeshComponent->GetCPUSkinnedVertices(FinalVertices, LODIndex);
 
 		FSkeletalMeshLODModel& LODModel = Resource->LODModels[LODIndex];
 		
 		const int32 NumSections = LODModel.Sections.Num();
 		
-		//Empty the raw mesh
-		RawMesh.Empty();
+		// Empty the mesh description
+		OutMeshDescription.Empty();
 
-		FStaticMeshAttributes Attributes(RawMesh);
+		FStaticMeshAttributes Attributes(OutMeshDescription);
 		TVertexAttributesRef<FVector> VertexPositions = Attributes.GetVertexPositions();
 		TEdgeAttributesRef<bool> EdgeHardnesses = Attributes.GetEdgeHardnesses();
 		TPolygonGroupAttributesRef<FName> PolygonGroupImportedMaterialSlotNames = Attributes.GetPolygonGroupMaterialSlotNames();
@@ -301,15 +303,15 @@ void FMeshMergeHelpers::RetrieveMesh(USkeletalMeshComponent* SkeletalMeshCompone
 			TotalTriangles += SkelMeshSection.NumTriangles;
 		}
 		TotalCorners = TotalTriangles * 3;
-		RawMesh.ReserveNewVertices(FinalVertices.Num());
-		RawMesh.ReserveNewPolygons(TotalTriangles);
-		RawMesh.ReserveNewVertexInstances(TotalCorners);
-		RawMesh.ReserveNewEdges(TotalCorners);
+		OutMeshDescription.ReserveNewVertices(FinalVertices.Num());
+		OutMeshDescription.ReserveNewPolygons(TotalTriangles);
+		OutMeshDescription.ReserveNewVertexInstances(TotalCorners);
+		OutMeshDescription.ReserveNewEdges(TotalCorners);
 
 		// Copy skinned vertex positions
 		for (int32 VertIndex = 0; VertIndex < FinalVertices.Num(); ++VertIndex)
 		{
-			const FVertexID VertexID = RawMesh.CreateVertex();
+			const FVertexID VertexID = OutMeshDescription.CreateVertex();
 			VertexPositions[VertexID] = FinalVertices[VertIndex].Position;
 		}
 
@@ -331,9 +333,9 @@ void FMeshMergeHelpers::RetrieveMesh(USkeletalMeshComponent* SkeletalMeshCompone
 
 			FName ImportedMaterialSlotName = SkeletalMeshComponent->SkeletalMesh->Materials[MaterialIndex].ImportedMaterialSlotName;
 			const FPolygonGroupID SectionPolygonGroupID(SectionIndex);
-			if (!RawMesh.IsPolygonGroupValid(SectionPolygonGroupID))
+			if (!OutMeshDescription.IsPolygonGroupValid(SectionPolygonGroupID))
 			{
-				RawMesh.CreatePolygonGroupWithID(SectionPolygonGroupID);
+				OutMeshDescription.CreatePolygonGroupWithID(SectionPolygonGroupID);
 				PolygonGroupImportedMaterialSlotNames[SectionPolygonGroupID] = ImportedMaterialSlotName;
 			}
 			int32 WedgeIndex = 0;
@@ -346,7 +348,7 @@ void FMeshMergeHelpers::RetrieveMesh(USkeletalMeshComponent* SkeletalMeshCompone
 				{
 					const int32 VertexIndexForWedge = LODModel.IndexBuffer[SkelMeshSection.BaseIndex + WedgeIndex];
 					VertexIndexes[CornerIndex] = FVertexID(VertexIndexForWedge);
-					FVertexInstanceID VertexInstanceID = RawMesh.CreateVertexInstance(VertexIndexes[CornerIndex]);
+					FVertexInstanceID VertexInstanceID = OutMeshDescription.CreateVertexInstance(VertexIndexes[CornerIndex]);
 					VertexInstanceIDs[CornerIndex] = VertexInstanceID;
 					
 					const FSoftSkinVertex& SoftVertex = SkelMeshSection.SoftVertices[VertexIndexForWedge - SkelMeshSection.BaseVertexIndex];
@@ -372,13 +374,13 @@ void FMeshMergeHelpers::RetrieveMesh(USkeletalMeshComponent* SkeletalMeshCompone
 					VertexInstanceColors[VertexInstanceID] = bPropagateVertexColours ? FVector4(FLinearColor(SoftVertex.Color)) : FVector4(1.0f, 1.0f, 1.0f);
 				}
 				//Create a polygon from this triangle
-				const FPolygonID NewPolygonID = RawMesh.CreatePolygon(SectionPolygonGroupID, VertexInstanceIDs);
+				const FPolygonID NewPolygonID = OutMeshDescription.CreatePolygon(SectionPolygonGroupID, VertexInstanceIDs);
 			}
 		}
 	}
 }
 
-void FMeshMergeHelpers::RetrieveMesh(const UStaticMesh* StaticMesh, int32 LODIndex, FMeshDescription& RawMesh)
+void FMeshMergeHelpers::RetrieveMesh(const UStaticMesh* StaticMesh, int32 LODIndex, FMeshDescription& OutMeshDescription)
 {
 	const FStaticMeshSourceModel& StaticMeshModel = StaticMesh->GetSourceModel(LODIndex);
 
@@ -387,20 +389,20 @@ void FMeshMergeHelpers::RetrieveMesh(const UStaticMesh* StaticMesh, int32 LODInd
 	
 	// Check whether or not this mesh has been reduced in-engine
 	const bool bReducedMesh = StaticMesh->IsReductionActive(LODIndex);
-	// Trying to retrieve rawmesh from SourceStaticMeshModel was giving issues, which causes a mismatch			
+	// Trying to retrieve mesh from SourceStaticMeshModel was giving issues, which causes a mismatch			
 	const bool bRenderDataMismatch = (LODIndex > 0) || StaticMeshModel.BuildSettings.bGenerateLightmapUVs;
 
 	if (bImportedMesh && !bReducedMesh && !bRenderDataMismatch)
 	{
-		RawMesh = *StaticMesh->GetMeshDescription(LODIndex);
+		OutMeshDescription = *StaticMesh->GetMeshDescription(LODIndex);
 	}
 	else
 	{
-		ExportStaticMeshLOD(StaticMesh->RenderData->LODResources[LODIndex], RawMesh, StaticMesh->StaticMaterials);
+		ExportStaticMeshLOD(StaticMesh->RenderData->LODResources[LODIndex], OutMeshDescription, StaticMesh->StaticMaterials);
 	}
 
-	// Make sure the raw mesh is not irreparably malformed.
-	if (RawMesh.VertexInstances().Num() <= 0)
+	// Make sure the mesh is not irreparably malformed.
+	if (OutMeshDescription.VertexInstances().Num() <= 0)
 	{
 		// wrong
 		bool check = true;
@@ -423,24 +425,24 @@ void FMeshMergeHelpers::RetrieveMesh(const UStaticMesh* StaticMesh, int32 LODInd
 	ComputeNTBsOptions |= (bImportedMesh && BuildSettings.bRecomputeNormals) ? EComputeNTBsFlags::Normals : EComputeNTBsFlags::None;
 	ComputeNTBsOptions |= (bImportedMesh && BuildSettings.bRecomputeTangents) ? EComputeNTBsFlags::Tangents : EComputeNTBsFlags::None;
 
-	FStaticMeshOperations::ComputeTriangleTangentsAndNormals(RawMesh, 0.0f);
-	FStaticMeshOperations::RecomputeNormalsAndTangentsIfNeeded(RawMesh, ComputeNTBsOptions);
+	FStaticMeshOperations::ComputeTriangleTangentsAndNormals(OutMeshDescription, 0.0f);
+	FStaticMeshOperations::RecomputeNormalsAndTangentsIfNeeded(OutMeshDescription, ComputeNTBsOptions);
 }
 
-void FMeshMergeHelpers::ExportStaticMeshLOD(const FStaticMeshLODResources& StaticMeshLOD, FMeshDescription& OutRawMesh, const TArray<FStaticMaterial>& Materials)
+void FMeshMergeHelpers::ExportStaticMeshLOD(const FStaticMeshLODResources& StaticMeshLOD, FMeshDescription& OutMeshDescription, const TArray<FStaticMaterial>& Materials)
 {
 	const int32 NumWedges = StaticMeshLOD.IndexBuffer.GetNumIndices();
 	const int32 NumVertexPositions = StaticMeshLOD.VertexBuffers.PositionVertexBuffer.GetNumVertices();
 	const int32 NumFaces = NumWedges / 3;
 
-	OutRawMesh.Empty();
+	OutMeshDescription.Empty();
 
 	if (NumVertexPositions <= 0 || StaticMeshLOD.VertexBuffers.StaticMeshVertexBuffer.GetNumVertices() <= 0)
 	{
 		return;
 	}
 
-	FStaticMeshAttributes Attributes(OutRawMesh);
+	FStaticMeshAttributes Attributes(OutMeshDescription);
 	TVertexAttributesRef<FVector> VertexPositions = Attributes.GetVertexPositions();
 	TEdgeAttributesRef<bool> EdgeHardnesses = Attributes.GetEdgeHardnesses();
 	TPolygonGroupAttributesRef<FName> PolygonGroupImportedMaterialSlotNames = Attributes.GetPolygonGroupMaterialSlotNames();
@@ -450,10 +452,10 @@ void FMeshMergeHelpers::ExportStaticMeshLOD(const FStaticMeshLODResources& Stati
 	TVertexInstanceAttributesRef<FVector4> VertexInstanceColors = Attributes.GetVertexInstanceColors();
 	TVertexInstanceAttributesRef<FVector2D> VertexInstanceUVs = Attributes.GetVertexInstanceUVs();
 
-	OutRawMesh.ReserveNewVertices(NumVertexPositions);
-	OutRawMesh.ReserveNewVertexInstances(NumWedges);
-	OutRawMesh.ReserveNewPolygons(NumFaces);
-	OutRawMesh.ReserveNewEdges(NumWedges);
+	OutMeshDescription.ReserveNewVertices(NumVertexPositions);
+	OutMeshDescription.ReserveNewVertexInstances(NumWedges);
+	OutMeshDescription.ReserveNewPolygons(NumFaces);
+	OutMeshDescription.ReserveNewEdges(NumWedges);
 
 	const int32 NumTexCoords = StaticMeshLOD.VertexBuffers.StaticMeshVertexBuffer.GetNumTexCoords();
 	VertexInstanceUVs.SetNumChannels(NumTexCoords);
@@ -461,7 +463,7 @@ void FMeshMergeHelpers::ExportStaticMeshLOD(const FStaticMeshLODResources& Stati
 	for (int32 SectionIndex = 0; SectionIndex < StaticMeshLOD.Sections.Num(); ++SectionIndex)
 	{
 		const FStaticMeshSection& Section = StaticMeshLOD.Sections[SectionIndex];
-		FPolygonGroupID CurrentPolygonGroupID = OutRawMesh.CreatePolygonGroup();
+		FPolygonGroupID CurrentPolygonGroupID = OutMeshDescription.CreatePolygonGroup();
 		check(CurrentPolygonGroupID.GetValue() == SectionIndex);
 		if (Materials.IsValidIndex(Section.MaterialIndex))
 		{
@@ -476,7 +478,7 @@ void FMeshMergeHelpers::ExportStaticMeshLOD(const FStaticMeshLODResources& Stati
 	//Create the vertex
 	for (int32 VertexIndex = 0; VertexIndex < NumVertexPositions; ++VertexIndex)
 	{
-		FVertexID VertexID = OutRawMesh.CreateVertex();
+		FVertexID VertexID = OutMeshDescription.CreateVertex();
 		VertexPositions[VertexID] = StaticMeshLOD.VertexBuffers.PositionVertexBuffer.VertexPosition(VertexIndex);
 	}
 
@@ -505,7 +507,7 @@ void FMeshMergeHelpers::ExportStaticMeshLOD(const FStaticMeshLODResources& Stati
 		{
 			int32 WedgeIndex = StaticMeshLOD.IndexBuffer.GetIndex(TriangleIndex * 3 + Corner);
 			FVertexID VertexID(WedgeIndex);
-			FVertexInstanceID VertexInstanceID = OutRawMesh.CreateVertexInstance(VertexID);
+			FVertexInstanceID VertexInstanceID = OutMeshDescription.CreateVertexInstance(VertexID);
 			VertexIDs[Corner] = VertexID;
 			VertexInstanceIDs[Corner] = VertexInstanceID;
 
@@ -534,7 +536,7 @@ void FMeshMergeHelpers::ExportStaticMeshLOD(const FStaticMeshLODResources& Stati
 			}
 		}
 		//Create a polygon from this triangle
-		const FPolygonID NewPolygonID = OutRawMesh.CreatePolygon(CurrentPolygonGroupID, VertexInstanceIDs);
+		const FPolygonID NewPolygonID = OutMeshDescription.CreatePolygon(CurrentPolygonGroupID, VertexInstanceIDs);
 	}
 }
 
@@ -580,7 +582,7 @@ bool FMeshMergeHelpers::CheckWrappingUVs(const FMeshDescription& MeshDescription
 	return bResult;
 }
 
-void FMeshMergeHelpers::CullTrianglesFromVolumesAndUnderLandscapes(const UWorld* World, const FBoxSphereBounds& Bounds, FMeshDescription& InOutRawMesh)
+void FMeshMergeHelpers::CullTrianglesFromVolumesAndUnderLandscapes(const UWorld* World, const FBoxSphereBounds& Bounds, FMeshDescription& InOutMeshDescription)
 {
 	TArray<ALandscapeProxy*> Landscapes;
 	TArray<AMeshMergeCullingVolume*> CullVolumes;
@@ -619,12 +621,12 @@ void FMeshMergeHelpers::CullTrianglesFromVolumesAndUnderLandscapes(const UWorld*
 		}
 	}
 
-	TVertexAttributesConstRef<FVector> VertexPositions = InOutRawMesh.GetVertexPositions();
+	TVertexAttributesConstRef<FVector> VertexPositions = InOutMeshDescription.GetVertexPositions();
 
 	TMap<FVertexID, bool> VertexVisible;
-	VertexVisible.Reserve(InOutRawMesh.Vertices().Num());
+	VertexVisible.Reserve(InOutMeshDescription.Vertices().Num());
 	int32 Index = 0;
-	for(const FVertexID VertexID : InOutRawMesh.Vertices().GetElementIDs())
+	for(const FVertexID VertexID : InOutMeshDescription.Vertices().GetElementIDs())
 	{
 		const FVector& Position = VertexPositions[VertexID];
 		// Start with setting visibility to true on all vertices
@@ -675,12 +677,12 @@ void FMeshMergeHelpers::CullTrianglesFromVolumesAndUnderLandscapes(const UWorld*
 
 	// We now know which vertices are below the landscape
 	TArray<FTriangleID> TriangleToRemove;
-	for(const FTriangleID TriangleID : InOutRawMesh.Triangles().GetElementIDs())
+	for(const FTriangleID TriangleID : InOutMeshDescription.Triangles().GetElementIDs())
 	{
 		bool AboveLandscape = false;
 		for (int32 Corner = 0; Corner < 3; ++Corner)
 		{
-			AboveLandscape |= VertexVisible[InOutRawMesh.GetVertexInstanceVertex(InOutRawMesh.GetTriangleVertexInstance(TriangleID, Corner))];
+			AboveLandscape |= VertexVisible[InOutMeshDescription.GetVertexInstanceVertex(InOutMeshDescription.GetTriangleVertexInstance(TriangleID, Corner))];
 		}
 
 		if (!AboveLandscape)
@@ -697,7 +699,7 @@ void FMeshMergeHelpers::CullTrianglesFromVolumesAndUnderLandscapes(const UWorld*
 		TArray<FVertexID> OrphanedVertices;
 		for (FTriangleID TriangleID : TriangleToRemove)
 		{
-			InOutRawMesh.DeleteTriangle(TriangleID, &OrphanedEdges, &OrphanedVertexInstances, &OrphanedPolygonGroups);
+			InOutMeshDescription.DeleteTriangle(TriangleID, &OrphanedEdges, &OrphanedVertexInstances, &OrphanedPolygonGroups);
 		}
 		//Do not remove the polygongroup since its indexed with the mesh material array
 		/*for (FPolygonGroupID PolygonGroupID : OrphanedPolygonGroups)
@@ -706,26 +708,26 @@ void FMeshMergeHelpers::CullTrianglesFromVolumesAndUnderLandscapes(const UWorld*
 		}*/
 		for (FVertexInstanceID VertexInstanceID : OrphanedVertexInstances)
 		{
-			InOutRawMesh.DeleteVertexInstance(VertexInstanceID, &OrphanedVertices);
+			InOutMeshDescription.DeleteVertexInstance(VertexInstanceID, &OrphanedVertices);
 		}
 		for (FEdgeID EdgeID : OrphanedEdges)
 		{
-			InOutRawMesh.DeleteEdge(EdgeID, &OrphanedVertices);
+			InOutMeshDescription.DeleteEdge(EdgeID, &OrphanedVertices);
 		}
 		for (FVertexID VertexID : OrphanedVertices)
 		{
-			InOutRawMesh.DeleteVertex(VertexID);
+			InOutMeshDescription.DeleteVertex(VertexID);
 		}
 		//Compact and Remap IDs so we have clean ID from 0 to n since we just erase some polygons
 		//The render build need to have compact ID
 		FElementIDRemappings OutRemappings;
-		InOutRawMesh.Compact(OutRemappings);
+		InOutMeshDescription.Compact(OutRemappings);
 	}
 }
 
-void FMeshMergeHelpers::PropagateSplineDeformationToRawMesh(const USplineMeshComponent* InSplineMeshComponent, FMeshDescription &OutRawMesh)
+void FMeshMergeHelpers::PropagateSplineDeformationToMesh(const USplineMeshComponent* InSplineMeshComponent, FMeshDescription& InOutMeshDescription)
 {
-	FStaticMeshAttributes Attributes(OutRawMesh);
+	FStaticMeshAttributes Attributes(InOutMeshDescription);
 
 	TVertexAttributesRef<FVector> VertexPositions = Attributes.GetVertexPositions();
 	TVertexInstanceAttributesRef<FVector> VertexInstanceNormals = Attributes.GetVertexInstanceNormals();
@@ -734,12 +736,12 @@ void FMeshMergeHelpers::PropagateSplineDeformationToRawMesh(const USplineMeshCom
 
 	// Apply spline deformation for each vertex's tangents
 	int32 WedgeIndex = 0;
-	for (const FTriangleID TriangleID : OutRawMesh.Triangles().GetElementIDs())
+	for (const FTriangleID TriangleID : InOutMeshDescription.Triangles().GetElementIDs())
 	{
 		for (int32 Corner = 0; Corner < 3; ++Corner, ++WedgeIndex)
 		{
-			const FVertexInstanceID VertexInstanceID = OutRawMesh.GetTriangleVertexInstance(TriangleID, Corner);
-			const FVertexID VertexID = OutRawMesh.GetVertexInstanceVertex(VertexInstanceID);
+			const FVertexInstanceID VertexInstanceID = InOutMeshDescription.GetTriangleVertexInstance(TriangleID, Corner);
+			const FVertexID VertexID = InOutMeshDescription.GetVertexInstanceVertex(VertexInstanceID);
 			const float& AxisValue = USplineMeshComponent::GetAxisValue(VertexPositions[VertexID], InSplineMeshComponent->ForwardAxis);
 			FTransform SliceTransform = InSplineMeshComponent->CalcSliceTransform(AxisValue);
 			FVector TangentY = FVector::CrossProduct(VertexInstanceNormals[VertexInstanceID], VertexInstanceTangents[VertexInstanceID]).GetSafeNormal() * VertexInstanceBinormalSigns[VertexInstanceID];
@@ -751,7 +753,7 @@ void FMeshMergeHelpers::PropagateSplineDeformationToRawMesh(const USplineMeshCom
 	}
 
 	// Apply spline deformation for each vertex position
-	for (const FVertexID VertexID : OutRawMesh.Vertices().GetElementIDs())
+	for (const FVertexID VertexID : InOutMeshDescription.Vertices().GetElementIDs())
 	{
 		float& AxisValue = USplineMeshComponent::GetAxisValue(VertexPositions[VertexID], InSplineMeshComponent->ForwardAxis);
 		FTransform SliceTransform = InSplineMeshComponent->CalcSliceTransform(AxisValue);
@@ -789,17 +791,17 @@ void FMeshMergeHelpers::PropagateSplineDeformationToPhysicsGeometry(USplineMeshC
 	}
 }
 
-void FMeshMergeHelpers::TransformRawMeshVertexData(const FTransform& InTransform, FMeshDescription &OutRawMesh)
+void FMeshMergeHelpers::TransformMeshDescriptionVertexData(const FTransform& InTransform, FMeshDescription& InOutMeshDescription)
 {
-	TRACE_CPUPROFILER_EVENT_SCOPE(FMeshMergeHelpers::TransformRawMeshVertexData)
+	TRACE_CPUPROFILER_EVENT_SCOPE(FMeshMergeHelpers::TransformMeshDescriptionVertexData)
 
-	FStaticMeshAttributes Attributes(OutRawMesh);
+	FStaticMeshAttributes Attributes(InOutMeshDescription);
 	TArrayView<FVector> VertexPositions = Attributes.GetVertexPositions().GetRawArray();
 	TArrayView<FVector> VertexInstanceNormals = Attributes.GetVertexInstanceNormals().GetRawArray();
 	TArrayView<FVector> VertexInstanceTangents = Attributes.GetVertexInstanceTangents().GetRawArray();
 	TArrayView<float> VertexInstanceBinormalSigns = Attributes.GetVertexInstanceBinormalSigns().GetRawArray();
 
-	for(const FVertexID VertexID : OutRawMesh.Vertices().GetElementIDs())
+	for(const FVertexID VertexID : InOutMeshDescription.Vertices().GetElementIDs())
 	{
 		VertexPositions[VertexID] = InTransform.TransformPosition(VertexPositions[VertexID]);
 	}
@@ -815,7 +817,7 @@ void FMeshMergeHelpers::TransformRawMeshVertexData(const FTransform& InTransform
 			Normal = AdjointT.TransformVector(Normal) * MulBy;
 		};
 
-	for (const FVertexInstanceID VertexInstanceID : OutRawMesh.VertexInstances().GetElementIDs())
+	for (const FVertexInstanceID VertexInstanceID : InOutMeshDescription.VertexInstances().GetElementIDs())
 	{
 		FVector TangentY = FVector::CrossProduct(VertexInstanceNormals[VertexInstanceID], VertexInstanceTangents[VertexInstanceID]).GetSafeNormal() * VertexInstanceBinormalSigns[VertexInstanceID];
 		TransformNormal(VertexInstanceTangents[VertexInstanceID]);
@@ -827,12 +829,12 @@ void FMeshMergeHelpers::TransformRawMeshVertexData(const FTransform& InTransform
 	const bool bIsMirrored = InTransform.GetDeterminant() < 0.f;
 	if (bIsMirrored)
 	{
-		//Reverse the vertex instance
-		OutRawMesh.ReverseAllPolygonFacing();
+		// Reverse the vertex instance
+		InOutMeshDescription.ReverseAllPolygonFacing();
 	}
 }
 
-void FMeshMergeHelpers::RetrieveCullingLandscapeAndVolumes(UWorld* InWorld, const FBoxSphereBounds& EstimatedMeshProxyBounds, const TEnumAsByte<ELandscapeCullingPrecision::Type> PrecisionType, TArray<FMeshDescription*>& CullingRawMeshes)
+void FMeshMergeHelpers::RetrieveCullingLandscapeAndVolumes(UWorld* InWorld, const FBoxSphereBounds& EstimatedMeshProxyBounds, const TEnumAsByte<ELandscapeCullingPrecision::Type> PrecisionType, TArray<FMeshDescription*>& OutCullingMeshes)
 {
 	// Extract landscape proxies and cull volumes from the world
 	TArray<ALandscapeProxy*> LandscapeActors;
@@ -881,7 +883,7 @@ void FMeshMergeHelpers::RetrieveCullingLandscapeAndVolumes(UWorld* InWorld, cons
 		Landscape->ExportToRawMesh(LandscapeExportLOD, *MeshDescription, LandscapeBounds);
 		if (MeshDescription->Vertices().Num())
 		{
-			CullingRawMeshes.Add(MeshDescription);
+			OutCullingMeshes.Add(MeshDescription);
 		}
 	}
 
@@ -904,7 +906,7 @@ void FMeshMergeHelpers::RetrieveCullingLandscapeAndVolumes(UWorld* InWorld, cons
 			VertexPositions[VertexID] += VolumeLocation;
 		}
 
-		CullingRawMeshes.Add(VolumeMesh);
+		OutCullingMeshes.Add(VolumeMesh);
 	}
 }
 
@@ -1034,17 +1036,17 @@ FVector2D FMeshMergeHelpers::GetValidUV(const FVector2D& UV)
 	return NewUV;
 }
 
-void FMeshMergeHelpers::CalculateTextureCoordinateBoundsForRawMesh(const FMeshDescription& InRawMesh, TArray<FBox2D>& OutBounds)
+void FMeshMergeHelpers::CalculateTextureCoordinateBoundsForMesh(const FMeshDescription& InMeshDescription, TArray<FBox2D>& OutBounds)
 {
-	TVertexInstanceAttributesConstRef<FVector2D> VertexInstanceUVs = FStaticMeshConstAttributes(InRawMesh).GetVertexInstanceUVs();
+	TVertexInstanceAttributesConstRef<FVector2D> VertexInstanceUVs = FStaticMeshConstAttributes(InMeshDescription).GetVertexInstanceUVs();
 	OutBounds.Empty();
-	for (const FTriangleID TriangleID : InRawMesh.Polygons().GetElementIDs())
+	for (const FTriangleID TriangleID : InMeshDescription.Polygons().GetElementIDs())
 	{
-		int32 MaterialIndex = InRawMesh.GetTrianglePolygonGroup(TriangleID).GetValue();
+		int32 MaterialIndex = InMeshDescription.GetTrianglePolygonGroup(TriangleID).GetValue();
 		if (OutBounds.Num() <= MaterialIndex)
 			OutBounds.SetNumZeroed(MaterialIndex + 1);
 		{
-			TArrayView<const FVertexInstanceID> TriangleVertexInstances = InRawMesh.GetTriangleVertexInstances(TriangleID);
+			TArrayView<const FVertexInstanceID> TriangleVertexInstances = InMeshDescription.GetTriangleVertexInstances(TriangleID);
 			for (const FVertexInstanceID VertexInstanceID : TriangleVertexInstances)
 			{
 				for (int32 UVIndex = 0; UVIndex < VertexInstanceUVs.GetNumChannels(); ++UVIndex)
@@ -1056,7 +1058,7 @@ void FMeshMergeHelpers::CalculateTextureCoordinateBoundsForRawMesh(const FMeshDe
 	}
 }
 
-bool FMeshMergeHelpers::PropagatePaintedColorsToRawMesh(const UStaticMeshComponent* StaticMeshComponent, int32 LODIndex, FMeshDescription& RawMesh)
+bool FMeshMergeHelpers::PropagatePaintedColorsToMesh(const UStaticMeshComponent* StaticMeshComponent, int32 LODIndex, FMeshDescription& InOutMeshDescription)
 {
 	UStaticMesh* StaticMesh = StaticMeshComponent->GetStaticMesh();
 
@@ -1069,11 +1071,11 @@ bool FMeshMergeHelpers::PropagatePaintedColorsToRawMesh(const UStaticMeshCompone
 
 		if (ColorVertexBuffer.GetNumVertices() == RenderModel.GetNumVertices())
 		{	
-			const int32 NumWedges = RawMesh.VertexInstances().Num();
+			const int32 NumWedges = InOutMeshDescription.VertexInstances().Num();
 			const int32 NumRenderWedges = RenderModel.IndexBuffer.GetNumIndices();
 			const bool bUseRenderWedges = NumWedges == NumRenderWedges;
 
-			TVertexInstanceAttributesRef<FVector4> VertexInstanceColors = FStaticMeshAttributes(RawMesh).GetVertexInstanceColors();
+			TVertexInstanceAttributesRef<FVector4> VertexInstanceColors = FStaticMeshAttributes(InOutMeshDescription).GetVertexInstanceColors();
 
 			if (bUseRenderWedges)
 			{
@@ -1081,11 +1083,11 @@ bool FMeshMergeHelpers::PropagatePaintedColorsToRawMesh(const UStaticMeshCompone
 				TMap<int32, FVertexInstanceID> IndexToVertexInstanceID;
 				IndexToVertexInstanceID.Reserve(NumWedges);
 				int32 CurrentWedgeIndex = 0;
-				for (const FTriangleID TriangleID : RawMesh.Triangles().GetElementIDs())
+				for (const FTriangleID TriangleID : InOutMeshDescription.Triangles().GetElementIDs())
 				{
 					for (int32 Corner = 0; Corner < 3; ++Corner, ++CurrentWedgeIndex)
 					{
-						IndexToVertexInstanceID.Add(CurrentWedgeIndex, RawMesh.GetTriangleVertexInstance(TriangleID, Corner));
+						IndexToVertexInstanceID.Add(CurrentWedgeIndex, InOutMeshDescription.GetTriangleVertexInstance(TriangleID, Corner));
 					}
 				}
 
@@ -1107,18 +1109,18 @@ bool FMeshMergeHelpers::PropagatePaintedColorsToRawMesh(const UStaticMeshCompone
 			// Use index buffer directly. Not sure this will happen with FMeshDescription
 			else
 			{
-				if (RawMesh.Vertices().Num() == ColorVertexBuffer.GetNumVertices())
+				if (InOutMeshDescription.Vertices().Num() == ColorVertexBuffer.GetNumVertices())
 				{
 					//Create a map index
 					TMap<FVertexID, int32> VertexIDToVertexIndex;
-					VertexIDToVertexIndex.Reserve(RawMesh.Vertices().Num());
+					VertexIDToVertexIndex.Reserve(InOutMeshDescription.Vertices().Num());
 					int32 CurrentVertexIndex = 0;
-					for (const FVertexID VertexID : RawMesh.Vertices().GetElementIDs())
+					for (const FVertexID VertexID : InOutMeshDescription.Vertices().GetElementIDs())
 					{
 						VertexIDToVertexIndex.Add(VertexID, CurrentVertexIndex++);
 					}
 
-					for (const FVertexID VertexID : RawMesh.Vertices().GetElementIDs())
+					for (const FVertexID VertexID : InOutMeshDescription.Vertices().GetElementIDs())
 					{
 						FColor WedgeColor = FColor::White;
 						uint32 VertIndex = VertexIDToVertexIndex[VertexID];
@@ -1127,7 +1129,7 @@ bool FMeshMergeHelpers::PropagatePaintedColorsToRawMesh(const UStaticMeshCompone
 						{
 							WedgeColor = ColorVertexBuffer.VertexColor(VertIndex);
 						}
-						TArrayView<const FVertexInstanceID> VertexInstances = RawMesh.GetVertexVertexInstanceIDs(VertexID);
+						TArrayView<const FVertexInstanceID> VertexInstances = InOutMeshDescription.GetVertexVertexInstanceIDs(VertexID);
 						for (const FVertexInstanceID& VertexInstanceID : VertexInstances)
 						{
 							VertexInstanceColors[VertexInstanceID] = FLinearColor(WedgeColor);
@@ -1170,9 +1172,9 @@ bool FMeshMergeHelpers::IsLandscapeHit(const FVector& RayOrigin, const FVector&
 	return bHitLandscape;
 }
 
-void FMeshMergeHelpers::AppendRawMesh(FMeshDescription& InTarget, const FMeshDescription& InSource)
+void FMeshMergeHelpers::AppendMesh(FMeshDescription& InTarget, const FMeshDescription& InSource)
 {
-	TRACE_CPUPROFILER_EVENT_SCOPE(FMeshMergeHelpers::AppendRawMesh)
+	TRACE_CPUPROFILER_EVENT_SCOPE(FMeshMergeHelpers::AppendMesh)
 
 	FStaticMeshConstAttributes SourceAttributes(InSource);
 	TVertexAttributesConstRef<FVector> SourceVertexPositions = SourceAttributes.GetVertexPositions();
@@ -1298,15 +1300,7 @@ void FMeshMergeHelpers::AppendRawMesh(FMeshDescription& InTarget, const FMeshDes
 	}
 }
 
-void FMeshMergeHelpers::ExtractImposterToRawMesh(const UStaticMeshComponent* InImposterComponent, FMeshDescription& InRawMesh)
-{
-	// Retrieve imposter LOD mesh	
-	const int32 LODIndex = InImposterComponent->GetStaticMesh()->GetNumLODs() - 1;
-	FStaticMeshAttributes(InRawMesh).Register();
-	FMeshMergeHelpers::RetrieveMesh(InImposterComponent->GetStaticMesh(), LODIndex, InRawMesh);
-}
-
-void FMeshMergeHelpers::MergeImpostersToRawMesh(TArray<const UStaticMeshComponent*> ImposterComponents, FMeshDescription& InRawMesh, const FVector& InPivot, int32 InBaseMaterialIndex, TArray<UMaterialInterface*>& OutImposterMaterials)
+void FMeshMergeHelpers::MergeImpostersToMesh(TArray<const UStaticMeshComponent*> ImposterComponents, FMeshDescription& InMeshDescription, const FVector& InPivot, int32 InBaseMaterialIndex, TArray<UMaterialInterface*>& OutImposterMaterials)
 {
 	TMap<UMaterialInterface*, FPolygonGroupID> ImposterMaterialToPolygonGroupID;
 	for (const UStaticMeshComponent* Component : ImposterComponents)
@@ -1366,10 +1360,10 @@ void FMeshMergeHelpers::MergeImpostersToRawMesh(TArray<const UStaticMeshComponen
 
 		TPolygonGroupAttributesRef<FName> SourcePolygonGroupImportedMaterialSlotNames = ImposterMeshAttributes.GetPolygonGroupMaterialSlotNames();
 
-		FStaticMeshAttributes Attributes(InRawMesh);
+		FStaticMeshAttributes Attributes(InMeshDescription);
 		TPolygonGroupAttributesRef<FName> TargetPolygonGroupImportedMaterialSlotNames = Attributes.GetPolygonGroupMaterialSlotNames();
 
-		//Add the missing polygon group ID to the target(InRawMesh)
+		//Add the missing polygon group ID to the target(InMeshDescription)
 		//Remap the source mesh(ImposterMesh) polygongroup to fit with the target polygon groups
 		TMap<FPolygonGroupID, FPolygonGroupID> RemapSourcePolygonGroup;
 		RemapSourcePolygonGroup.Reserve(ImposterMesh.PolygonGroups().Num());
@@ -1381,7 +1375,7 @@ void FMeshMergeHelpers::MergeImpostersToRawMesh(TArray<const UStaticMeshComponen
 			FPolygonGroupID MatchTargetPolygonGroupID = ExistTargetPolygonGroupID == nullptr ? INDEX_NONE : *ExistTargetPolygonGroupID;
 			if (MatchTargetPolygonGroupID == INDEX_NONE)
 			{
-				MatchTargetPolygonGroupID = InRawMesh.CreatePolygonGroup();
+				MatchTargetPolygonGroupID = InMeshDescription.CreatePolygonGroup();
 				//use the material name to fill the imported material name. Material name will be unique
 				TargetPolygonGroupImportedMaterialSlotNames[MatchTargetPolygonGroupID] = MaterialUseBySection->GetFName();
 				ImposterMaterialToPolygonGroupID.Add(MaterialUseBySection, MatchTargetPolygonGroupID);
@@ -1390,6 +1384,6 @@ void FMeshMergeHelpers::MergeImpostersToRawMesh(TArray<const UStaticMeshComponen
 		}
 		ImposterMesh.RemapPolygonGroups(RemapSourcePolygonGroup);
 
-		FMeshMergeHelpers::AppendRawMesh(InRawMesh, ImposterMesh);
+		AppendMesh(InMeshDescription, ImposterMesh);
 	}
 }
diff --git a/Engine/Source/Developer/MeshMergeUtilities/Private/MeshMergeHelpers.h b/Engine/Source/Developer/MeshMergeUtilities/Private/MeshMergeHelpers.h
--- a/Engine/Source/Developer/MeshMergeUtilities/Private/MeshMergeHelpers.h
+++ b/Engine/Source/Developer/MeshMergeUtilities/Private/MeshMergeHelpers.h
@@ -31,12 +31,12 @@ public:
 	static void ExpandInstances(const UInstancedStaticMeshComponent* InInstancedStaticMeshComponent, FMeshDescription& InOutRawMesh, TArray<FSectionInfo>& InOutSections);
 
 	/** Extracting mesh data in FMeshDescription form from static, skeletal mesh (components) */
-	static void RetrieveMesh(const UStaticMeshComponent* StaticMeshComponent, int32 LODIndex, FMeshDescription& RawMesh, bool bPropagateVertexColours);
-	static void RetrieveMesh(USkeletalMeshComponent* SkeletalMeshComponent, int32 LODIndex, FMeshDescription& RawMesh, bool bPropagateVertexColours);
-	static void RetrieveMesh(const UStaticMesh* StaticMesh, int32 LODIndex, FMeshDescription& RawMesh);
+	static void RetrieveMesh(const UStaticMeshComponent* StaticMeshComponent, int32 LODIndex, FMeshDescription& OutMeshDescription, bool bPropagateVertexColours);
+	static void RetrieveMesh(const USkeletalMeshComponent* SkeletalMeshComponent, int32 LODIndex, FMeshDescription& OutMeshDescription, bool bPropagateVertexColours);
+	static void RetrieveMesh(const UStaticMesh* StaticMesh, int32 LODIndex, FMeshDescription& OutMeshDescription);
 	
 	/** Exports static mesh LOD render data to a RawMesh */
-	static void ExportStaticMeshLOD(const FStaticMeshLODResources& StaticMeshLOD, FMeshDescription& OutRawMesh, const TArray<FStaticMaterial>& Materials);
+	static void ExportStaticMeshLOD(const FStaticMeshLODResources& StaticMeshLOD, FMeshDescription& OutMeshDescription, const TArray<FStaticMaterial>& Materials);
 
 	/** Checks whether or not the texture coordinates are outside of 0-1 UV ranges */
 	static bool CheckWrappingUVs(const TArray<FVector2D>& UVs);
@@ -45,17 +45,17 @@ public:
 	/** Culls away triangles which are inside culling volumes or completely underneath the landscape */
 	static void CullTrianglesFromVolumesAndUnderLandscapes(const UWorld* World, const FBoxSphereBounds& Bounds, FMeshDescription& InOutRawMesh);
 	
-	/** Propagates deformation along spline to raw mesh data */
-	static void PropagateSplineDeformationToRawMesh(const USplineMeshComponent* InSplineMeshComponent, FMeshDescription &OutRawMesh);
+	/** Propagates deformation along spline to mesh data */
+	static void PropagateSplineDeformationToMesh(const USplineMeshComponent* InSplineMeshComponent, FMeshDescription& InOutMeshDescription);
 	
 	/** Propagates deformation along spline to physics geometry data */
 	static void PropagateSplineDeformationToPhysicsGeometry(USplineMeshComponent* SplineMeshComponent, FKAggregateGeom& InOutPhysicsGeometry);
 
-	/** Transforms raw mesh data using InTransform*/
-	static void TransformRawMeshVertexData(const FTransform& InTransform, FMeshDescription &OutRawMesh);
+	/** Transforms mesh data using InTransform*/
+	static void TransformMeshDescriptionVertexData(const FTransform& InTransform, FMeshDescription& OutMeshDescInOutMeshDescriptionription);
 
-	/** Retrieves all culling landscapes and volumes as FMeshDescription structures. Note the caller is responsible for deleting the heap data managed by CullingRawMeshes */
-	static void RetrieveCullingLandscapeAndVolumes(UWorld* InWorld, const FBoxSphereBounds& EstimatedMeshProxyBounds, const TEnumAsByte<ELandscapeCullingPrecision::Type> PrecisionType, TArray<FMeshDescription*>& CullingRawMeshes);
+	/** Retrieves all culling landscapes and volumes as FMeshDescription structures. Note the caller is responsible for deleting the heap data managed by OutCullingMeshes */
+	static void RetrieveCullingLandscapeAndVolumes(UWorld* InWorld, const FBoxSphereBounds& EstimatedMeshProxyBounds, const TEnumAsByte<ELandscapeCullingPrecision::Type> PrecisionType, TArray<FMeshDescription*>& OutCullingMeshes);
 
 	/** Transforms physics geometry data using InTransform */
 	static void TransformPhysicsGeometry(const FTransform& InTransform, const bool bBakeConvexTransform, struct FKAggregateGeom& AggGeom);
@@ -66,21 +66,18 @@ public:
 	/** Ensure that UV is in valid 0-1 UV ranges */
 	static FVector2D GetValidUV(const FVector2D& UV);
 	
-	/** Calculates UV coordinates bounds for the given Raw Mesh	*/
-	static void CalculateTextureCoordinateBoundsForRawMesh(const FMeshDescription& InRawMesh, TArray<FBox2D>& OutBounds);
+	/** Calculates UV coordinates bounds for the given MeshDescription	*/
+	static void CalculateTextureCoordinateBoundsForMesh(const FMeshDescription& InMeshDescription, TArray<FBox2D>& OutBounds);
 
-	/** Propagates vertex painted colors from the StaticMeshComponent instance to RawMesh */
-	static bool PropagatePaintedColorsToRawMesh(const UStaticMeshComponent* StaticMeshComponent, int32 LODIndex, FMeshDescription& RawMesh);
+	/** Propagates vertex painted colors from the StaticMeshComponent instance to MeshDescription */
+	static bool PropagatePaintedColorsToMesh(const UStaticMeshComponent* StaticMeshComponent, int32 LODIndex, FMeshDescription& InOutMeshDescription);
 
 	/** Checks whether or not the landscape proxy is hit given a ray start and end */
 	static bool IsLandscapeHit(const FVector& RayOrigin, const FVector& RayEndPoint, const UWorld* World, const TArray<ALandscapeProxy*>& LandscapeProxies, FVector& OutHitLocation);
 	
 	/** Appends a FMeshDescription to another instance */
-	static void AppendRawMesh(FMeshDescription& InTarget, const FMeshDescription& InSource);
+	static void AppendMesh(FMeshDescription& InTarget, const FMeshDescription& InSource);
 
-	/** Merges imposter meshes into a raw mesh. */
-	static void MergeImpostersToRawMesh(TArray<const UStaticMeshComponent*> ImposterComponents, FMeshDescription& InRawMesh, const FVector& InPivot, int32 BaseMaterialIndex, TArray<UMaterialInterface*>& OutImposterMaterials);
-
-	/** Extract imposter meshes to a raw mesh. */
-	static void ExtractImposterToRawMesh(const UStaticMeshComponent* InImposterComponent, FMeshDescription& InRawMesh);
+	/** Merges imposter meshes into the given MeshDescription. */
+	static void MergeImpostersToMesh(TArray<const UStaticMeshComponent*> ImposterComponents, FMeshDescription& InOutMeshDescription, const FVector& InPivot, int32 BaseMaterialIndex, TArray<UMaterialInterface*>& OutImposterMaterials);
 };
diff --git a/Engine/Source/Developer/MeshMergeUtilities/Private/MeshMergeUtilities.cpp b/Engine/Source/Developer/MeshMergeUtilities/Private/MeshMergeUtilities.cpp
--- a/Engine/Source/Developer/MeshMergeUtilities/Private/MeshMergeUtilities.cpp
+++ b/Engine/Source/Developer/MeshMergeUtilities/Private/MeshMergeUtilities.cpp
@@ -1619,7 +1619,7 @@ void FMeshMergeUtilities::CreateProxyMesh(const TArray<UStaticMeshComponent*>& I
 			MergeData.RawMesh = MeshDescriptionData[Index];
 			MergeData.bIsClippingMesh = false;
 
-			FMeshMergeHelpers::CalculateTextureCoordinateBoundsForRawMesh(*MergeData.RawMesh, MergeData.TexCoordBounds);
+			FMeshMergeHelpers::CalculateTextureCoordinateBoundsForMesh(*MergeData.RawMesh, MergeData.TexCoordBounds);
 
 			FMeshData* MeshData = GlobalMeshSettings.FindByPredicate([&](const FMeshData& Entry)
 			{
@@ -1787,6 +1787,21 @@ bool FMeshMergeUtilities::IsValidBaseMaterial(const UMaterialInterface* InBaseMa
 	return false;
 }
 
+void FMeshMergeUtilities::RetrieveMeshDescription(const UStaticMeshComponent* InStaticMeshComponent, int32 LODIndex, FMeshDescription& InOutMeshDescription, bool bPropagateMeshData) const
+{
+	FMeshMergeHelpers::RetrieveMesh(InStaticMeshComponent, LODIndex, InOutMeshDescription, bPropagateMeshData);
+}
+
+void FMeshMergeUtilities::RetrieveMeshDescription(const USkeletalMeshComponent* InSkeletalMeshComponent, int32 LODIndex, FMeshDescription& InOutMeshDescription, bool bPropagateMeshData) const
+{
+	FMeshMergeHelpers::RetrieveMesh(InSkeletalMeshComponent, LODIndex, InOutMeshDescription, bPropagateMeshData);
+}
+
+void FMeshMergeUtilities::RetrieveMeshDescription(const UStaticMesh* InStaticMesh, int32 LODIndex, FMeshDescription& InOutMeshDescription) const
+{
+	FMeshMergeHelpers::RetrieveMesh(InStaticMesh, LODIndex, InOutMeshDescription);
+}
+
 void FMeshMergeUtilities::RegisterExtension(IMeshMergeExtension* InExtension)
 {
 	MeshMergeExtensions.Add(InExtension);
@@ -2735,7 +2750,7 @@ void FMeshMergeUtilities::MergeComponentsToStaticMesh(const TArray<UPrimitiveCom
 				if (bContainsImposters)
 				{
 					// Merge imposter meshes to rawmesh
-					FMeshMergeHelpers::MergeImpostersToRawMesh(ImposterComponents, MergedMeshLOD, MergedAssetPivot, UniqueMaterials.Num(), ImposterMaterials);					
+					FMeshMergeHelpers::MergeImpostersToMesh(ImposterComponents, MergedMeshLOD, MergedAssetPivot, UniqueMaterials.Num(), ImposterMaterials);					
 
 					const FTransform PivotTransform = FTransform(MergedAssetPivot);
 					for (const UStaticMeshComponent* Component : ImposterComponents)
@@ -2865,12 +2880,6 @@ void FMeshMergeUtilities::MergeComponentsToStaticMesh(const TArray<UPrimitiveCom
 	}
 }
 
-void FMeshMergeUtilities::ExtractImposterToRawMesh(const UStaticMeshComponent* InImposterComponent, FMeshDescription& InImposterMesh) const
-{
-	check(InImposterComponent->bUseMaxLODAsImposter);
-	FMeshMergeHelpers::ExtractImposterToRawMesh(InImposterComponent, InImposterMesh);
-}
-
 void FMeshMergeUtilities::CreateMergedRawMeshes(FMeshMergeDataTracker& InDataTracker, const FMeshMergingSettings& InSettings, const TArray<UStaticMeshComponent*>& InStaticMeshComponentsToMerge, const TArray<UMaterialInterface*>& InUniqueMaterials, const TMap<UMaterialInterface*, UMaterialInterface*>& InCollapsedMaterialMap, const TMultiMap<FMeshLODKey, MaterialRemapPair>& InOutputMaterialsMap, bool bInMergeAllLODs, bool bInMergeMaterialData, const FVector& InMergedAssetPivot, TArray<FMeshDescription>& OutMergedRawMeshes) const
 {
 	TRACE_CPUPROFILER_EVENT_SCOPE(FMeshMergeUtilities::CreateMergedRawMeshes)
diff --git a/Engine/Source/Developer/MeshMergeUtilities/Private/MeshMergeUtilities.h b/Engine/Source/Developer/MeshMergeUtilities/Private/MeshMergeUtilities.h
--- a/Engine/Source/Developer/MeshMergeUtilities/Private/MeshMergeUtilities.h
+++ b/Engine/Source/Developer/MeshMergeUtilities/Private/MeshMergeUtilities.h
@@ -34,10 +34,12 @@ public:
 	UE_DEPRECATED(4.20, "Use different signature containing BaseMaterial parameter")
 	virtual void CreateProxyMesh(const TArray<UStaticMeshComponent*>& InStaticMeshComps, const struct FMeshProxySettings& InMeshProxySettings, UPackage* InOuter, const FString& InProxyBasePackageName, const FGuid InGuid, const FCreateProxyDelegate& InProxyCreatedDelegate, const bool bAllowAsync = false, const float ScreenSize = 1.0f) const override;
 
-	virtual void ExtractImposterToRawMesh(const UStaticMeshComponent* InImposterComponent, FMeshDescription& OutImposterMesh) const override;
-
 	virtual bool IsValidBaseMaterial(const UMaterialInterface* InBaseMaterial, bool bShowToaster) const override;
 
+	virtual void RetrieveMeshDescription(const UStaticMeshComponent* InStaticMeshComponent, int32 LODIndex, FMeshDescription& InOutMeshDescription, bool bPropagateVertexColours) const override;
+	virtual void RetrieveMeshDescription(const USkeletalMeshComponent* InSkeletalMeshComponent, int32 LODIndex, FMeshDescription& InOutMeshDescription, bool bPropagateVertexColours) const override;
+	virtual void RetrieveMeshDescription(const UStaticMesh* InStaticMesh, int32 LODIndex, FMeshDescription& InOutMeshDescription) const override;
+
 	virtual void RegisterExtension(IMeshMergeExtension* InExtension) override;
 	virtual void UnregisterExtension(IMeshMergeExtension* InExtension) override;	
 protected:
diff --git a/Engine/Source/Developer/MeshMergeUtilities/Private/ProxyGenerationProcessor.cpp b/Engine/Source/Developer/MeshMergeUtilities/Private/ProxyGenerationProcessor.cpp
--- a/Engine/Source/Developer/MeshMergeUtilities/Private/ProxyGenerationProcessor.cpp
+++ b/Engine/Source/Developer/MeshMergeUtilities/Private/ProxyGenerationProcessor.cpp
@@ -248,7 +248,7 @@ void FProxyGenerationProcessor::ProcessJob(const FGuid& JobGuid, FProxyGeneratio
 
 		// Merge imposter meshes to rawmesh
 		// The base material index is always one here as we assume we only have one HLOD material
-		FMeshMergeHelpers::MergeImpostersToRawMesh(Data->MergeData->ImposterComponents, Data->RawMesh, FVector::ZeroVector, 1, ImposterMaterials);
+		FMeshMergeHelpers::MergeImpostersToMesh(Data->MergeData->ImposterComponents, Data->RawMesh, FVector::ZeroVector, 1, ImposterMaterials);
 
 		for (const UStaticMeshComponent* Component : Data->MergeData->ImposterComponents)
 		{
diff --git a/Engine/Source/Developer/MeshMergeUtilities/Public/IMeshMergeUtilities.h b/Engine/Source/Developer/MeshMergeUtilities/Public/IMeshMergeUtilities.h
--- a/Engine/Source/Developer/MeshMergeUtilities/Public/IMeshMergeUtilities.h
+++ b/Engine/Source/Developer/MeshMergeUtilities/Public/IMeshMergeUtilities.h
@@ -125,17 +125,16 @@ public:
 	*/
 	virtual void CreateProxyMesh(const TArray<UStaticMeshComponent*>& InComponents, const struct FMeshProxySettings& InMeshProxySettings, UMaterialInterface* InBaseMaterial, UPackage* InOuter, const FString& InProxyBasePackageName, const FGuid InGuid, const FCreateProxyDelegate& InProxyCreatedDelegate, const bool bAllowAsync = false, const float ScreenSize = 1.0f) const = 0;
 
-    /**
-	*	Retrieve the mesh description for an imposter static mesh component.
-	*
-	*	@param	InImposterComponent		The static mesh component whose static mesh lowest LOD is an imposter mesh.
-	*	@param	OutImposterMesh			The extracted mesh description.
-	*/
-	virtual void ExtractImposterToRawMesh(const UStaticMeshComponent* InImposterComponent, FMeshDescription& OutImposterMesh) const = 0;
-
 	/** Checks whether or not the give material is valid as a base for the Proxy Material. This as it requires certain Material Parameters to be included to get a desirable material */
 	virtual bool IsValidBaseMaterial(const UMaterialInterface* InBaseMaterial, bool bShowToaster) const = 0;
-	
+
+	/** Extract mesh data in FMeshDescription form from the provided static mesh component */
+	virtual void RetrieveMeshDescription(const UStaticMeshComponent* InStaticMeshComponent, int32 LODIndex, FMeshDescription& InOutMeshDescription, bool bPropagateVertexColours) const = 0;
+	/** Extract mesh data in FMeshDescription form from the provided skeletal mesh component */
+	virtual void RetrieveMeshDescription(const USkeletalMeshComponent* InSkeletalMeshComponent, int32 LODIndex, FMeshDescription& InOutMeshDescription, bool bPropagateVertexColours) const = 0;
+	/** Extract mesh data in FMeshDescription form from the provided static mesh */
+	virtual void RetrieveMeshDescription(const UStaticMesh* InStaticMesh, int32 LODIndex, FMeshDescription& InOutMeshDescription) const = 0;
+
 	virtual void RegisterExtension(IMeshMergeExtension* InExtension) = 0;
 	virtual void UnregisterExtension(IMeshMergeExtension* InExtension) = 0;
 };
\ No newline at end of file
