Checking for any LandscapeProxy actors when merging actors with "Use Landscape Culling" option selected

diff --git a/Engine/Source/Developer/MeshMergeUtilities/Private/MeshMergeHelpers.cpp b/Engine/Source/Developer/MeshMergeUtilities/Private/MeshMergeHelpers.cpp
--- a/Engine/Source/Developer/MeshMergeUtilities/Private/MeshMergeHelpers.cpp
+++ b/Engine/Source/Developer/MeshMergeUtilities/Private/MeshMergeHelpers.cpp
@@ -593,17 +593,17 @@ void FMeshMergeHelpers::CullTrianglesFromVolumesAndUnderLandscapes(const UWorld*
 	{
 		for (AActor* Actor : Level->Actors)
 		{
-			ALandscape* Proxy = Cast<ALandscape>(Actor);
-			if (Proxy && Proxy->bUseLandscapeForCullingInvisibleHLODVertices)
+			ALandscapeProxy* LandscapeProxy = Cast<ALandscapeProxy>(Actor);
+			if (LandscapeProxy && LandscapeProxy->bUseLandscapeForCullingInvisibleHLODVertices)
 			{
 				FVector Origin, Extent;
-				Proxy->GetActorBounds(false, Origin, Extent);
+				LandscapeProxy->GetActorBounds(false, Origin, Extent);
 				FBox LandscapeBox(Origin - Extent, Origin + Extent);
 
 				// Ignore Z axis for 2d bounds check
 				if (LandscapeBox.IntersectXY(BoxBounds))
 				{
-					Landscapes.Add(Proxy->GetLandscapeActor());
+					Landscapes.Add(LandscapeProxy);
 				}
 			}
 
