Optimized HLOD proxy mesh generation
* Avoid recreating render states for all primitives in the world

diff --git a/Engine/Source/Developer/MeshMergeUtilities/Private/ProxyGenerationProcessor.cpp b/Engine/Source/Developer/MeshMergeUtilities/Private/ProxyGenerationProcessor.cpp
--- a/Engine/Source/Developer/MeshMergeUtilities/Private/ProxyGenerationProcessor.cpp
+++ b/Engine/Source/Developer/MeshMergeUtilities/Private/ProxyGenerationProcessor.cpp
@@ -144,6 +144,10 @@ void FProxyGenerationProcessor::ProcessJob(const FGuid& JobGuid, FProxyGeneratio
 
 	if (!Data->RawMesh.IsEmpty())
 	{
+		// Don't recreate render states with the material update context as we will manually do it through
+		// the FStaticMeshComponentRecreateRenderStateContext below
+		FMaterialUpdateContext MaterialUpdateContext(FMaterialUpdateContext::EOptions::Default & ~FMaterialUpdateContext::EOptions::RecreateRenderStates);
+
 		// Retrieve flattened material data
 		FFlattenMaterial& FlattenMaterial = Data->Material;
 
@@ -154,7 +158,7 @@ void FProxyGenerationProcessor::ProcessJob(const FGuid& JobGuid, FProxyGeneratio
 		FMaterialUtilities::OptimizeFlattenMaterial(FlattenMaterial);
 
 		// Create a new proxy material instance
-		ProxyMaterial = ProxyMaterialUtilities::CreateProxyMaterialInstance(Data->MergeData->InOuter, Data->MergeData->InProxySettings.MaterialSettings, Data->MergeData->BaseMaterial, FlattenMaterial, AssetBasePath, AssetBaseName, OutAssetsToSync);
+		ProxyMaterial = ProxyMaterialUtilities::CreateProxyMaterialInstance(Data->MergeData->InOuter, Data->MergeData->InProxySettings.MaterialSettings, Data->MergeData->BaseMaterial, FlattenMaterial, AssetBasePath, AssetBaseName, OutAssetsToSync, &MaterialUpdateContext);
 
 		for (IMeshMergeExtension* Extension : Owner->MeshMergeExtensions)
 		{
diff --git a/Engine/Source/Developer/MeshMergeUtilities/Private/ProxyMaterialUtilities.h b/Engine/Source/Developer/MeshMergeUtilities/Private/ProxyMaterialUtilities.h
--- a/Engine/Source/Developer/MeshMergeUtilities/Private/ProxyMaterialUtilities.h
+++ b/Engine/Source/Developer/MeshMergeUtilities/Private/ProxyMaterialUtilities.h
@@ -102,7 +102,7 @@ namespace ProxyMaterialUtilities
 		return NumPacked >= 2;
 	}
 
-	static UMaterialInstanceConstant* CreateProxyMaterialInstance(UPackage* InOuter, const FMaterialProxySettings& InMaterialProxySettings, UMaterialInterface* InBaseMaterial, const FFlattenMaterial& FlattenMaterial, const FString& AssetBasePath, const FString& AssetBaseName, TArray<UObject*>& OutAssetsToSync)
+	static UMaterialInstanceConstant* CreateProxyMaterialInstance(UPackage* InOuter, const FMaterialProxySettings& InMaterialProxySettings, UMaterialInterface* InBaseMaterial, const FFlattenMaterial& FlattenMaterial, const FString& AssetBasePath, const FString& AssetBaseName, TArray<UObject*>& OutAssetsToSync, FMaterialUpdateContext* MaterialUpdateContext = nullptr)
 	{
 		UMaterialInterface* BaseMaterial = InBaseMaterial;
 		
@@ -267,7 +267,7 @@ namespace ProxyMaterialUtilities
 		}
 
 		// Force initializing the static permutations according to the switches we have set
-		OutMaterial->UpdateStaticPermutation(NewStaticParameterSet);
+		OutMaterial->UpdateStaticPermutation(NewStaticParameterSet, MaterialUpdateContext);
 		OutMaterial->InitStaticPermutation();
 
 		OutMaterial->PostEditChange();
@@ -275,10 +275,10 @@ namespace ProxyMaterialUtilities
 		return OutMaterial;
 	}	
 
-	static UMaterialInstanceConstant* CreateProxyMaterialInstance(UPackage* InOuter, const FMaterialProxySettings& InMaterialProxySettings, FFlattenMaterial& FlattenMaterial, const FString& AssetBasePath, const FString& AssetBaseName, TArray<UObject*>& OutAssetsToSync)
+	static UMaterialInstanceConstant* CreateProxyMaterialInstance(UPackage* InOuter, const FMaterialProxySettings& InMaterialProxySettings, FFlattenMaterial& FlattenMaterial, const FString& AssetBasePath, const FString& AssetBaseName, TArray<UObject*>& OutAssetsToSync, FMaterialUpdateContext* MaterialUpdateContext = nullptr)
 	{
 		UMaterialInterface* BaseMaterial = LoadObject<UMaterialInterface>(NULL, TEXT("/Engine/EngineMaterials/BaseFlattenMaterial.BaseFlattenMaterial"), NULL, LOAD_None, NULL);
-		return CreateProxyMaterialInstance(InOuter, InMaterialProxySettings, BaseMaterial, FlattenMaterial, AssetBasePath, AssetBasePath, OutAssetsToSync);
+		return CreateProxyMaterialInstance(InOuter, InMaterialProxySettings, BaseMaterial, FlattenMaterial, AssetBasePath, AssetBasePath, OutAssetsToSync, MaterialUpdateContext);
 	}
 
 };
