Optimized HLOD merged mesh generation
* Avoid recreating render states for all primitives in the world

diff --git a/Engine/Source/Developer/MeshMergeUtilities/Private/MeshMergeUtilities.cpp b/Engine/Source/Developer/MeshMergeUtilities/Private/MeshMergeUtilities.cpp
--- a/Engine/Source/Developer/MeshMergeUtilities/Private/MeshMergeUtilities.cpp
+++ b/Engine/Source/Developer/MeshMergeUtilities/Private/MeshMergeUtilities.cpp
@@ -2555,7 +2555,11 @@ void FMeshMergeUtilities::MergeComponentsToStaticMesh(const TArray<UPrimitiveCom
 
 		OutMaterial.UVChannel = MergedMatUVChannel;
 
-		MergedMaterial = CreateProxyMaterial(InBasePackageName, MergedAssetPackageName, InBaseMaterial, InOuter, InSettings, OutMaterial, OutAssetsToSync);
+		// Don't recreate render states with the material update context as we will manually do it through
+		// the FStaticMeshComponentRecreateRenderStateContext used below at the creation of the static mesh.
+		FMaterialUpdateContext MaterialUpdateContext(FMaterialUpdateContext::EOptions::Default & ~FMaterialUpdateContext::EOptions::RecreateRenderStates);
+
+		MergedMaterial = CreateProxyMaterial(InBasePackageName, MergedAssetPackageName, InBaseMaterial, InOuter, InSettings, OutMaterial, OutAssetsToSync, &MaterialUpdateContext);
 		
 		if (MergedMaterial)
 		{
@@ -3369,7 +3373,7 @@ void FMeshMergeUtilities::MergeComponentsToInstances(const TArray<UPrimitiveComp
 	}
 }
 
-UMaterialInterface* FMeshMergeUtilities::CreateProxyMaterial(const FString &InBasePackageName, FString MergedAssetPackageName, UMaterialInterface* InBaseMaterial, UPackage* InOuter, const FMeshMergingSettings &InSettings, FFlattenMaterial OutMaterial, TArray<UObject *>& OutAssetsToSync) const
+UMaterialInterface* FMeshMergeUtilities::CreateProxyMaterial(const FString &InBasePackageName, FString MergedAssetPackageName, UMaterialInterface* InBaseMaterial, UPackage* InOuter, const FMeshMergingSettings &InSettings, FFlattenMaterial OutMaterial, TArray<UObject *>& OutAssetsToSync, FMaterialUpdateContext* InMaterialUpdateContext) const
 {
 	// Create merged material asset
 	FString MaterialAssetName;
@@ -3394,7 +3398,7 @@ UMaterialInterface* FMeshMergeUtilities::CreateProxyMaterial(const FString &InBa
 		MaterialPackage->Modify();
 	}
 
-	UMaterialInstanceConstant* MergedMaterial = ProxyMaterialUtilities::CreateProxyMaterialInstance(MaterialPackage, InSettings.MaterialSettings, InBaseMaterial, OutMaterial, MaterialPackageName, MaterialAssetName, OutAssetsToSync);
+	UMaterialInstanceConstant* MergedMaterial = ProxyMaterialUtilities::CreateProxyMaterialInstance(MaterialPackage, InSettings.MaterialSettings, InBaseMaterial, OutMaterial, MaterialPackageName, MaterialAssetName, OutAssetsToSync, InMaterialUpdateContext);
 	// Set material static lighting usage flag if project has static lighting enabled
 	static const auto AllowStaticLightingVar = IConsoleManager::Get().FindTConsoleVariableDataInt(TEXT("r.AllowStaticLighting"));
 	const bool bAllowStaticLighting = (!AllowStaticLightingVar || AllowStaticLightingVar->GetValueOnGameThread() != 0);
diff --git a/Engine/Source/Developer/MeshMergeUtilities/Private/MeshMergeUtilities.h b/Engine/Source/Developer/MeshMergeUtilities/Private/MeshMergeUtilities.h
--- a/Engine/Source/Developer/MeshMergeUtilities/Private/MeshMergeUtilities.h
+++ b/Engine/Source/Developer/MeshMergeUtilities/Private/MeshMergeUtilities.h
@@ -49,7 +49,7 @@ protected:
 	/** Determines whether or not an individual material uses model vertex data during the shading process and outputs per-material flags */
 	void DetermineMaterialVertexDataUsage(TArray<bool>& InOutMaterialUsesVertexData, const TArray<UMaterialInterface*>& UniqueMaterials, const UMaterialOptions* MaterialOptions) const;
 	/** Creates a proxy material instance at givne path and name */
-	UMaterialInterface* CreateProxyMaterial(const FString &InBasePackageName, FString MergedAssetPackageName, UMaterialInterface* InBaseMaterial, UPackage* InOuter, const FMeshMergingSettings &InSettings, FFlattenMaterial OutMaterial, TArray<UObject *>& OutAssetsToSync) const;	
+	UMaterialInterface* CreateProxyMaterial(const FString &InBasePackageName, FString MergedAssetPackageName, UMaterialInterface* InBaseMaterial, UPackage* InOuter, const FMeshMergingSettings &InSettings, FFlattenMaterial OutMaterial, TArray<UObject *>& OutAssetsToSync, FMaterialUpdateContext* InMaterialUpdateContext = nullptr) const;
 	/** Scales texture coordinates to the specified box, e.g. to 0-1 range in U and V */
 	void ScaleTextureCoordinatesToBox(const FBox2D& Box, TArray<FVector2D>& InOutTextureCoordinates) const;
 	/** Merges flattened material into atlas textures */
