Output an error if calling FMeshMergeUtilities::CreateProxyMesh() without having the necessary plugin enabled

diff --git a/Engine/Source/Developer/MeshMergeUtilities/Private/MeshMergeUtilities.cpp b/Engine/Source/Developer/MeshMergeUtilities/Private/MeshMergeUtilities.cpp
--- a/Engine/Source/Developer/MeshMergeUtilities/Private/MeshMergeUtilities.cpp
+++ b/Engine/Source/Developer/MeshMergeUtilities/Private/MeshMergeUtilities.cpp
@@ -1146,14 +1146,14 @@ void FMeshMergeUtilities::CreateProxyMesh(const TArray<UStaticMeshComponent*>& I
 	// Error/warning checking for input
 	if (ReductionModule.GetMeshMergingInterface() == nullptr)
 	{
-		UE_LOG(LogMeshMerging, Log, TEXT("No automatic mesh merging module available"));
+		UE_LOG(LogMeshMerging, Error, TEXT("No mesh reduction module available. You must enable a plugin that provides that functionality (ex: ProxyLODPlugin)"));
 		return;
 	}
 
 	// Check that the delegate has a func-ptr bound to it
 	if (!InProxyCreatedDelegate.IsBound())
 	{
-		UE_LOG(LogMeshMerging, Log, TEXT("Invalid (unbound) delegate for returning generated proxy mesh"));
+		UE_LOG(LogMeshMerging, Warning, TEXT("Invalid (unbound) delegate for returning generated proxy mesh"));
 		return;
 	}
 
