Fix proxy mesh generation failing with sections that have no triangles

diff --git a/Engine/Source/Developer/MeshMergeUtilities/Private/MeshMergeHelpers.cpp b/Engine/Source/Developer/MeshMergeUtilities/Private/MeshMergeHelpers.cpp
--- a/Engine/Source/Developer/MeshMergeUtilities/Private/MeshMergeHelpers.cpp
+++ b/Engine/Source/Developer/MeshMergeUtilities/Private/MeshMergeHelpers.cpp
@@ -72,12 +72,6 @@ void FMeshMergeHelpers::ExtractSections(const UStaticMeshComponent* Component, i
 
 	for (const FStaticMeshSection& MeshSection : StaticMesh->GetRenderData()->LODResources[LODIndex].Sections)
 	{
-		// Skip empty sections
-		if (MeshSection.NumTriangles == 0)
-		{
-			continue;
-		}
-
 		// Retrieve material for this section
 		UMaterialInterface* StoredMaterial = Component->GetMaterial(MeshSection.MaterialIndex);
 
