Optimized proxy mesh generation
* Added GroupIdenticalMeshesForBaking, which, when activated, will result in identical meshes (or mesh instances) being baked only once.

Other changes/fixes
* Always rely on the StaticMesh RenderData when retrieving mesh description for mesh merging/simplification (was the case when retrieving from a SMC, now also the case for a SM).
* For some reason, Material baking was rendering each triangle twice, but with opposite facing. Uneeded as we are rendering with backface culling disabled.

diff --git a/Engine/Plugins/Importers/USDImporter/Source/USDUtilities/Private/USDShadeConversion.cpp b/Engine/Plugins/Importers/USDImporter/Source/USDUtilities/Private/USDShadeConversion.cpp
--- a/Engine/Plugins/Importers/USDImporter/Source/USDUtilities/Private/USDShadeConversion.cpp
+++ b/Engine/Plugins/Importers/USDImporter/Source/USDUtilities/Private/USDShadeConversion.cpp
@@ -884,7 +884,7 @@ namespace UsdShadeConversionImpl
 		}
 
 		FMeshData MeshSettings;
-		MeshSettings.RawMeshDescription = nullptr;
+		MeshSettings.MeshDescription = nullptr;
 		MeshSettings.TextureCoordinateBox = FBox2D( FVector2D( 0.0f, 0.0f ), FVector2D( 1.0f, 1.0f ) );
 		MeshSettings.TextureCoordinateIndex = 0;
 
diff --git a/Engine/Source/Developer/MaterialBaking/Private/MaterialBakingModule.cpp b/Engine/Source/Developer/MaterialBaking/Private/MaterialBakingModule.cpp
--- a/Engine/Source/Developer/MaterialBaking/Private/MaterialBakingModule.cpp
+++ b/Engine/Source/Developer/MaterialBaking/Private/MaterialBakingModule.cpp
@@ -30,6 +30,8 @@
 
 IMPLEMENT_MODULE(FMaterialBakingModule, MaterialBaking);
 
+DEFINE_LOG_CATEGORY_STATIC(LogMaterialBaking, All, All);
+
 #define LOCTEXT_NAMESPACE "MaterialBakingModule"
 
 /** Cvars for advanced features */
@@ -349,6 +351,15 @@ void FMaterialBakingModule::BakeMaterials(const TArray<FMaterialData*>& Material
 
 void FMaterialBakingModule::BakeMaterials(const TArray<FMaterialDataEx*>& MaterialSettings, const TArray<FMeshData*>& MeshSettings, TArray<FBakeOutputEx>& Output)
 {
+	UE_LOG(LogMaterialBaking, Verbose, TEXT("Performing material baking for %d materials"), MaterialSettings.Num());
+	for (int32 i = 0; i < MaterialSettings.Num(); i++)
+	{
+		if (MaterialSettings[i]->Material && MeshSettings[i]->MeshDescription)
+		{
+			UE_LOG(LogMaterialBaking, Verbose, TEXT("    [%5d] Material: %-50s Vertices: %8d    Triangles: %8d"), i, *MaterialSettings[i]->Material->GetName(), MeshSettings[i]->MeshDescription->Vertices().Num(), MeshSettings[i]->MeshDescription->Triangles().Num());
+		}
+	}
+
 	RenderCaptureInterface::FScopedCapture RenderCapture(CVarMaterialBakingRDOCCapture.GetValueOnAnyThread() == 1, TEXT("MaterialBaking"));
 
 	TRACE_CPUPROFILER_EVENT_SCOPE(FMaterialBakingModule::BakeMaterials)
@@ -376,7 +387,7 @@ void FMaterialBakingModule::BakeMaterials(const TArray<FMaterialDataEx*>& Materi
 	// Soft page faults are now incredibly expensive on Windows 10.
 	Algo::SortBy(
 		ProcessingOrder,
-		[&MeshSettings](const uint32 Index){ return MeshSettings[Index]->RawMeshDescription ? MeshSettings[Index]->RawMeshDescription->Vertices().Num() : 0; },
+		[&MeshSettings](const uint32 Index){ return MeshSettings[Index]->MeshDescription ? MeshSettings[Index]->MeshDescription->Vertices().Num() : 0; },
 		TGreater<>()
 	);
 
@@ -507,6 +518,8 @@ void FMaterialBakingModule::BakeMaterials(const TArray<FMaterialDataEx*>& Materi
 		const FMeshData* CurrentMeshSettings = MeshSettings[MaterialIndex];
 		FBakeOutputEx& CurrentOutput = Output[MaterialIndex];
 
+		TRACE_CPUPROFILER_EVENT_SCOPE_TEXT(*CurrentMaterialSettings->Material->GetName())
+
 		TArray<FMaterialPropertyEx> MaterialPropertiesToBakeOut;
 		CurrentMaterialSettings->PropertySizes.GenerateKeyArray(MaterialPropertiesToBakeOut);
 
@@ -526,6 +539,9 @@ void FMaterialBakingModule::BakeMaterials(const TArray<FMaterialDataEx*>& Materi
 			for (int32 PropertyIndex = 0; PropertyIndex < NumPropertiesToRender; ++PropertyIndex)
 			{
 				const FMaterialPropertyEx& Property = MaterialPropertiesToBakeOut[PropertyIndex];
+
+				TRACE_CPUPROFILER_EVENT_SCOPE_TEXT(*Property.ToString())
+
 				FExportMaterialProxy* ExportMaterialProxy = CreateMaterialProxy(CurrentMaterialSettings->Material, Property);
 
 				if (!ExportMaterialProxy->IsCompilationFinished())
diff --git a/Engine/Source/Developer/MaterialBaking/Private/MaterialRenderItem.cpp b/Engine/Source/Developer/MaterialBaking/Private/MaterialRenderItem.cpp
--- a/Engine/Source/Developer/MaterialBaking/Private/MaterialRenderItem.cpp
+++ b/Engine/Source/Developer/MaterialBaking/Private/MaterialRenderItem.cpp
@@ -85,7 +85,7 @@ void FMeshMaterialRenderItem::GenerateRenderData()
 	// Reset array without resizing
 	Vertices.SetNum(0, false);
 	Indices.SetNum(0, false);
-	if (MeshSettings->RawMeshDescription)
+	if (MeshSettings->MeshDescription)
 	{
 		// Use supplied FMeshDescription data to populate render data
 		PopulateWithMeshData();
@@ -182,7 +182,7 @@ void FMeshMaterialRenderItem::PopulateWithQuadData()
 
 void FMeshMaterialRenderItem::PopulateWithMeshData()
 {
-	const FMeshDescription* RawMesh = MeshSettings->RawMeshDescription;
+	const FMeshDescription* RawMesh = MeshSettings->MeshDescription;
 
 	FStaticMeshConstAttributes Attributes(*RawMesh);
 	TArrayView<const FVector> VertexPositions = Attributes.GetVertexPositions().GetRawArray();
@@ -266,16 +266,13 @@ void FMeshMaterialRenderItem::PopulateWithMeshData()
 				Vert->TextureCoordinate[6].X = VertexPositions[SrcVertexID].X;
 				Vert->TextureCoordinate[6].Y = VertexPositions[SrcVertexID].Y;
 				Vert->TextureCoordinate[7].X = VertexPositions[SrcVertexID].Z;
+				Vert->TextureCoordinate[7].Y = 0.0f;
 
 				Vert->Color = FLinearColor(VertexInstanceColors[SrcVertexInstanceID]).ToFColor(true);
 				// add index
 				Indices.Add(VertIndex);
 				VertIndex++;
 			}
-			// add the same triangle with opposite vertex order
-			Indices.Add(VertIndex - 3);
-			Indices.Add(VertIndex - 1);
-			Indices.Add(VertIndex - 2);
 		}
 		FaceIndex++;
 	}
diff --git a/Engine/Source/Developer/MaterialBaking/Public/MaterialBakingStructures.h b/Engine/Source/Developer/MaterialBaking/Public/MaterialBakingStructures.h
--- a/Engine/Source/Developer/MaterialBaking/Public/MaterialBakingStructures.h
+++ b/Engine/Source/Developer/MaterialBaking/Public/MaterialBakingStructures.h
@@ -45,14 +45,14 @@ struct FMaterialDataEx
 struct FMeshData
 {
 	FMeshData()
-		: RawMeshDescription(nullptr), Mesh(nullptr), bMirrored(false), VertexColorHash(0), TextureCoordinateIndex(0), LightMapIndex(0), LightMap(nullptr), LightmapResourceCluster(nullptr)
+		: MeshDescription(nullptr), Mesh(nullptr), bMirrored(false), VertexColorHash(0), TextureCoordinateIndex(0), LightMapIndex(0), LightMap(nullptr), LightmapResourceCluster(nullptr)
 	{}
 
 	/** Ptr to raw mesh data to use for baking out the material data, if nullptr a standard quad is used */
-	FMeshDescription* RawMeshDescription;
+	const FMeshDescription* MeshDescription;
 
 	/** Ptr to original static mesh this mesh data came from */
-	UStaticMesh* Mesh;
+	const UStaticMesh* Mesh;
 
 	/** Transform determinant used to detect mirroring */
 	bool bMirrored;
diff --git a/Engine/Source/Developer/MaterialUtilities/Private/MaterialUtilities.cpp b/Engine/Source/Developer/MaterialUtilities/Private/MaterialUtilities.cpp
--- a/Engine/Source/Developer/MaterialUtilities/Private/MaterialUtilities.cpp
+++ b/Engine/Source/Developer/MaterialUtilities/Private/MaterialUtilities.cpp
@@ -2322,7 +2322,7 @@ bool FMaterialUtilities::ExportMaterials(TArray<FMaterialMergeData*>& MergeData,
 		FMaterialMergeData* CurrentMergeData = MergeData[MaterialIndex];
 
 		FMeshData* MeshSet = new FMeshData();
-		MeshSet->RawMeshDescription = const_cast<FMeshDescription*>(CurrentMergeData->Mesh);
+		MeshSet->MeshDescription = CurrentMergeData->Mesh;
 		MeshSet->TextureCoordinateBox = CurrentMergeData->TexcoordBounds;
 		MeshSet->CustomTextureCoordinates = CurrentMergeData->TexCoords;
 		MeshSettings.Add(MeshSet);
diff --git a/Engine/Source/Developer/MeshMergeUtilities/Private/MeshMergeHelpers.cpp b/Engine/Source/Developer/MeshMergeUtilities/Private/MeshMergeHelpers.cpp
--- a/Engine/Source/Developer/MeshMergeUtilities/Private/MeshMergeHelpers.cpp
+++ b/Engine/Source/Developer/MeshMergeUtilities/Private/MeshMergeHelpers.cpp
@@ -181,7 +181,7 @@ void FMeshMergeHelpers::ExtractSections(const UStaticMesh* StaticMesh, int32 LOD
 	}
 }
 
-void FMeshMergeHelpers::ExpandInstances(const UInstancedStaticMeshComponent* InInstancedStaticMeshComponent, FMeshDescription& InOutMeshDescription, TArray<FSectionInfo>& InOutSections)
+void FMeshMergeHelpers::ExpandInstances(const UInstancedStaticMeshComponent* InInstancedStaticMeshComponent, FMeshDescription& InOutMeshDescription)
 {
 	TRACE_CPUPROFILER_EVENT_SCOPE(FMeshMergeHelpers::ExpandInstances)
 
@@ -191,11 +191,26 @@ void FMeshMergeHelpers::ExpandInstances(const UInstancedStaticMeshComponent* InI
 	FTransform ComponentTransform = InInstancedStaticMeshComponent->GetComponentTransform();
 	FTransform ComponentTransformInv = ComponentTransform.Inverse();
 
+	// Reserve memory upfront
+	int32 NumInstances = InInstancedStaticMeshComponent->GetInstanceCount();
+
+	CombinedMeshDescription.ReserveNewVertices(InOutMeshDescription.Vertices().Num() * NumInstances);
+	CombinedMeshDescription.ReserveNewVertexInstances(InOutMeshDescription.VertexInstances().Num() * NumInstances);
+	CombinedMeshDescription.ReserveNewEdges(InOutMeshDescription.Edges().Num() * NumInstances);
+	CombinedMeshDescription.ReserveNewTriangles(InOutMeshDescription.Triangles().Num() * NumInstances);
+
+	FStaticMeshOperations::FAppendSettings AppendSettings;
+	for (int32 ChannelIdx = 0; ChannelIdx < FStaticMeshOperations::FAppendSettings::MAX_NUM_UV_CHANNELS; ++ChannelIdx)
+	{
+		AppendSettings.bMergeUVChannels[ChannelIdx] = true;
+	}
+
+	FMeshDescription InstanceMeshDescription;
 	for(const FInstancedStaticMeshInstanceData& InstanceData : InInstancedStaticMeshComponent->PerInstanceSMData)
 	{
-		FMeshDescription InstanceMeshDescription = InOutMeshDescription;
+		InstanceMeshDescription = InOutMeshDescription;
 		FStaticMeshOperations::ApplyTransform(InstanceMeshDescription, ComponentTransformInv * FTransform(InstanceData.Transform) * ComponentTransform);
-		AppendMesh(CombinedMeshDescription, InstanceMeshDescription);
+		FStaticMeshOperations::AppendMeshDescription(InstanceMeshDescription, CombinedMeshDescription, AppendSettings);
 	}
 
 	InOutMeshDescription = CombinedMeshDescription;
@@ -208,11 +223,9 @@ void FMeshMergeHelpers::RetrieveMesh(const UStaticMeshComponent* StaticMeshCompo
 	const UStaticMesh* StaticMesh = StaticMeshComponent->GetStaticMesh();
 	const FStaticMeshSourceModel& StaticMeshModel = StaticMesh->GetSourceModel(LODIndex);
 
-	const bool bIsSplineMeshComponent = StaticMeshComponent->IsA<USplineMeshComponent>();
-
 	// Imported meshes will have a valid mesh description
 	const bool bImportedMesh = StaticMesh->IsMeshDescriptionValid(LODIndex);
-		
+
 	// Export the mesh data using static mesh render data
 	ExportStaticMeshLOD(StaticMesh->GetRenderData()->LODResources[LODIndex], OutMeshDescription, StaticMesh->GetStaticMaterials());
 
@@ -222,13 +235,11 @@ void FMeshMergeHelpers::RetrieveMesh(const UStaticMeshComponent* StaticMeshCompo
 		return;
 	}
 
-	// Use build settings from base mesh for LOD entries that was generated inside Editor.
-	const FMeshBuildSettings& BuildSettings = bImportedMesh ? StaticMeshModel.BuildSettings : StaticMesh->GetSourceModel(0).BuildSettings;
-
 	// Transform mesh to world space
 	FTransform ComponentToWorldTransform = StaticMeshComponent->GetComponentTransform();
 
 	// Handle spline mesh deformation
+	const bool bIsSplineMeshComponent = StaticMeshComponent->IsA<USplineMeshComponent>();
 	if (bIsSplineMeshComponent)
 	{
 		const USplineMeshComponent* SplineMeshComponent = Cast<USplineMeshComponent>(StaticMeshComponent);
@@ -245,10 +256,8 @@ void FMeshMergeHelpers::RetrieveMesh(const UStaticMeshComponent* StaticMeshCompo
 	// Transform raw mesh vertex data by the Static Mesh Component's component to world transformation	
 	FStaticMeshOperations::ApplyTransform(OutMeshDescription, ComponentToWorldTransform);
 
-	if (OutMeshDescription.VertexInstances().Num() <= 0)
-	{
-		return;
-	}
+	// Use build settings from base mesh for LOD entries that were generated inside Editor.
+	const FMeshBuildSettings& BuildSettings = bImportedMesh ? StaticMeshModel.BuildSettings : StaticMesh->GetSourceModel(0).BuildSettings;
 
 	// Figure out if we should recompute normals and tangents. By default generated LODs should not recompute normals	
 	EComputeNTBsFlags ComputeNTBsOptions = EComputeNTBsFlags::BlendOverlappingNormals;
@@ -261,6 +270,7 @@ void FMeshMergeHelpers::RetrieveMesh(const UStaticMeshComponent* StaticMeshCompo
 	{
 		ComputeNTBsOptions |= EComputeNTBsFlags::UseMikkTSpace;
 	}
+
 	FStaticMeshOperations::ComputeTriangleTangentsAndNormals(OutMeshDescription, 0.0f);
 	FStaticMeshOperations::RecomputeNormalsAndTangentsIfNeeded(OutMeshDescription, ComputeNTBsOptions);
 }
@@ -382,33 +392,23 @@ void FMeshMergeHelpers::RetrieveMesh(const USkeletalMeshComponent* SkeletalMeshC
 
 void FMeshMergeHelpers::RetrieveMesh(const UStaticMesh* StaticMesh, int32 LODIndex, FMeshDescription& OutMeshDescription)
 {
+	TRACE_CPUPROFILER_EVENT_SCOPE(FMeshMergeHelpers::RetrieveMesh)
+
 	const FStaticMeshSourceModel& StaticMeshModel = StaticMesh->GetSourceModel(LODIndex);
 
 	// Imported meshes will have a valid mesh description
 	const bool bImportedMesh = StaticMesh->IsMeshDescriptionValid(LODIndex);
-	
-	// Check whether or not this mesh has been reduced in-engine
-	const bool bReducedMesh = StaticMesh->IsReductionActive(LODIndex);
-	// Trying to retrieve mesh from SourceStaticMeshModel was giving issues, which causes a mismatch			
-	const bool bRenderDataMismatch = (LODIndex > 0) || StaticMeshModel.BuildSettings.bGenerateLightmapUVs;
 
-	if (bImportedMesh && !bReducedMesh && !bRenderDataMismatch)
-	{
-		OutMeshDescription = *StaticMesh->GetMeshDescription(LODIndex);
-	}
-	else
-	{
-		ExportStaticMeshLOD(StaticMesh->GetRenderData()->LODResources[LODIndex], OutMeshDescription, StaticMesh->GetStaticMaterials());
-	}
+	// Export the mesh data using static mesh render data
+	ExportStaticMeshLOD(StaticMesh->GetRenderData()->LODResources[LODIndex], OutMeshDescription, StaticMesh->GetStaticMaterials());
 
 	// Make sure the mesh is not irreparably malformed.
 	if (OutMeshDescription.VertexInstances().Num() <= 0)
 	{
-		// wrong
-		bool check = true;
+		return;
 	}
 
-	// Use build settings from base mesh for LOD entries that was generated inside Editor.
+	// Use build settings from base mesh for LOD entries that were generated inside Editor.
 	const FMeshBuildSettings& BuildSettings = bImportedMesh ? StaticMeshModel.BuildSettings : StaticMesh->GetSourceModel(0).BuildSettings;
 
 	// Figure out if we should recompute normals and tangents. By default generated LODs should not recompute normals	
@@ -422,8 +422,6 @@ void FMeshMergeHelpers::RetrieveMesh(const UStaticMesh* StaticMesh, int32 LODInd
 	{
 		ComputeNTBsOptions |= EComputeNTBsFlags::UseMikkTSpace;
 	}
-	ComputeNTBsOptions |= (bImportedMesh && BuildSettings.bRecomputeNormals) ? EComputeNTBsFlags::Normals : EComputeNTBsFlags::None;
-	ComputeNTBsOptions |= (bImportedMesh && BuildSettings.bRecomputeTangents) ? EComputeNTBsFlags::Tangents : EComputeNTBsFlags::None;
 
 	FStaticMeshOperations::ComputeTriangleTangentsAndNormals(OutMeshDescription, 0.0f);
 	FStaticMeshOperations::RecomputeNormalsAndTangentsIfNeeded(OutMeshDescription, ComputeNTBsOptions);
@@ -1129,134 +1127,6 @@ bool FMeshMergeHelpers::IsLandscapeHit(const FVector& RayOrigin, const FVector&
 	return bHitLandscape;
 }
 
-void FMeshMergeHelpers::AppendMesh(FMeshDescription& InTarget, const FMeshDescription& InSource)
-{
-	TRACE_CPUPROFILER_EVENT_SCOPE(FMeshMergeHelpers::AppendMesh)
-
-	FStaticMeshConstAttributes SourceAttributes(InSource);
-	TVertexAttributesConstRef<FVector> SourceVertexPositions = SourceAttributes.GetVertexPositions();
-	TEdgeAttributesConstRef<bool> SourceEdgeHardnesses = SourceAttributes.GetEdgeHardnesses();
-	TPolygonGroupAttributesConstRef<FName> SourcePolygonGroupImportedMaterialSlotNames = SourceAttributes.GetPolygonGroupMaterialSlotNames();
-	TVertexInstanceAttributesConstRef<FVector> SourceVertexInstanceNormals = SourceAttributes.GetVertexInstanceNormals();
-	TVertexInstanceAttributesConstRef<FVector> SourceVertexInstanceTangents = SourceAttributes.GetVertexInstanceTangents();
-	TVertexInstanceAttributesConstRef<float> SourceVertexInstanceBinormalSigns = SourceAttributes.GetVertexInstanceBinormalSigns();
-	TVertexInstanceAttributesConstRef<FVector4> SourceVertexInstanceColors = SourceAttributes.GetVertexInstanceColors();
-	TVertexInstanceAttributesConstRef<FVector2D> SourceVertexInstanceUVs = SourceAttributes.GetVertexInstanceUVs();
-
-
-	FStaticMeshAttributes TargetAttributes(InTarget);
-	TVertexAttributesRef<FVector> TargetVertexPositions = TargetAttributes.GetVertexPositions();
-	TEdgeAttributesRef<bool> TargetEdgeHardnesses = TargetAttributes.GetEdgeHardnesses();
-	TPolygonGroupAttributesRef<FName> TargetPolygonGroupImportedMaterialSlotNames = TargetAttributes.GetPolygonGroupMaterialSlotNames();
-	TVertexInstanceAttributesRef<FVector> TargetVertexInstanceNormals = TargetAttributes.GetVertexInstanceNormals();
-	TVertexInstanceAttributesRef<FVector> TargetVertexInstanceTangents = TargetAttributes.GetVertexInstanceTangents();
-	TVertexInstanceAttributesRef<float> TargetVertexInstanceBinormalSigns = TargetAttributes.GetVertexInstanceBinormalSigns();
-	TVertexInstanceAttributesRef<FVector4> TargetVertexInstanceColors = TargetAttributes.GetVertexInstanceColors();
-	TVertexInstanceAttributesRef<FVector2D> TargetVertexInstanceUVs = TargetAttributes.GetVertexInstanceUVs();
-
-	InTarget.ReserveNewVertices(InSource.Vertices().Num());
-	InTarget.ReserveNewVertexInstances(InSource.VertexInstances().Num());
-	InTarget.ReserveNewEdges(InSource.Edges().Num());
-	InTarget.ReserveNewPolygons(InSource.Polygons().Num());
-
-	//Append PolygonGroup
-	for (const FPolygonGroupID SourcePolygonGroupID : InSource.PolygonGroups().GetElementIDs())
-	{
-		if (!InTarget.IsPolygonGroupValid(SourcePolygonGroupID))
-		{
-			const FName BaseName = SourcePolygonGroupImportedMaterialSlotNames[SourcePolygonGroupID];
-			FName CurrentTestName = BaseName;
-			int32 UniqueID = 1;
-			bool bUnique = true;
-			do 
-			{
-				bUnique = true;
-				for (const FPolygonGroupID PolygonGroupID : InTarget.PolygonGroups().GetElementIDs())
-				{
-					if (TargetPolygonGroupImportedMaterialSlotNames[PolygonGroupID] == CurrentTestName)
-					{
-						CurrentTestName = FName(*(BaseName.ToString() + FString::FromInt(UniqueID++)));
-						bUnique = false;
-						break;
-					}
-				}
-			} while (!bUnique);
-
-			InTarget.CreatePolygonGroupWithID(SourcePolygonGroupID);
-			TargetPolygonGroupImportedMaterialSlotNames[SourcePolygonGroupID] = CurrentTestName;
-		}
-	}
-
-	//Append the Vertexs
-	TMap<FVertexID, FVertexID> SourceToTargetVertexID;
-	SourceToTargetVertexID.Reserve(InSource.Vertices().Num());
-	for (const FVertexID SourceVertexID : InSource.Vertices().GetElementIDs())
-	{
-		const FVertexID TargetVertexID = InTarget.CreateVertex();
-		SourceToTargetVertexID.Add(SourceVertexID, TargetVertexID);
-		TargetVertexPositions[TargetVertexID] = SourceVertexPositions[SourceVertexID];
-	}
-
-	//Append VertexInstances
-	if (SourceVertexInstanceUVs.GetNumChannels() > TargetVertexInstanceUVs.GetNumChannels())
-	{
-		TargetVertexInstanceUVs.SetNumChannels(SourceVertexInstanceUVs.GetNumChannels());
-	}
-	TMap<FVertexInstanceID, FVertexInstanceID> SourceToTargetVertexInstanceID;
-	SourceToTargetVertexInstanceID.Reserve(InSource.VertexInstances().Num());
-	for (const FVertexInstanceID SourceVertexInstanceID : InSource.VertexInstances().GetElementIDs())
-	{
-		const FVertexID SourceVertexID = InSource.GetVertexInstanceVertex(SourceVertexInstanceID);
-		const FVertexInstanceID TargetVertexInstanceID = InTarget.CreateVertexInstance(SourceToTargetVertexID[SourceVertexID]);
-		TargetVertexInstanceTangents[TargetVertexInstanceID] = SourceVertexInstanceTangents[SourceVertexInstanceID];
-		TargetVertexInstanceBinormalSigns[TargetVertexInstanceID] = SourceVertexInstanceBinormalSigns[SourceVertexInstanceID];
-		TargetVertexInstanceNormals[TargetVertexInstanceID] = SourceVertexInstanceNormals[SourceVertexInstanceID];
-		TargetVertexInstanceColors[TargetVertexInstanceID] = SourceVertexInstanceColors[SourceVertexInstanceID];
-		for (int32 UVIndex = 0; UVIndex < TargetVertexInstanceUVs.GetNumChannels(); ++UVIndex)
-		{
-			FVector2D SourceUV = SourceVertexInstanceUVs.GetNumChannels() > UVIndex ? SourceVertexInstanceUVs.Get(SourceVertexInstanceID, UVIndex) : FVector2D(0.0f, 0.0f);
-			TargetVertexInstanceUVs.Set(TargetVertexInstanceID, UVIndex, SourceUV);
-		}
-		SourceToTargetVertexInstanceID.Add(SourceVertexInstanceID, TargetVertexInstanceID);
-	}
-
-	//Append Edges
-	TMap<FEdgeID, FEdgeID> SourceToTargetEdgeID;
-	SourceToTargetEdgeID.Reserve(InSource.Edges().Num());
-	for (const FEdgeID SourceEdgeID : InSource.Edges().GetElementIDs())
-	{
-		const FVertexID EdgeVertex0 = InSource.GetEdgeVertex(SourceEdgeID, 0);
-		const FVertexID EdgeVertex1 = InSource.GetEdgeVertex(SourceEdgeID, 1);
-		const FEdgeID TargetEdgeID = InTarget.CreateEdge(SourceToTargetVertexID[EdgeVertex0], SourceToTargetVertexID[EdgeVertex1]);
-		TargetEdgeHardnesses[TargetEdgeID] = SourceEdgeHardnesses[SourceEdgeID];
-		SourceToTargetEdgeID.Add(SourceEdgeID, TargetEdgeID);
-	}
-
-	auto CreateContour = [&InSource, &SourceToTargetVertexInstanceID, &SourceToTargetEdgeID](const TArray<FVertexInstanceID>& SourceVertexInstanceIDs, TArray<FVertexInstanceID>& DestVertexInstanceIDs)
-	{
-		const int32 ContourCount = SourceVertexInstanceIDs.Num();
-		for (int32 ContourIndex = 0; ContourIndex < ContourCount; ++ContourIndex)
-		{
-			FVertexInstanceID SourceVertexInstanceID = SourceVertexInstanceIDs[ContourIndex];
-			DestVertexInstanceIDs.Add(SourceToTargetVertexInstanceID[SourceVertexInstanceID]);
-		}
-	};
-
-	// @todo: make this work in terms of triangles instead of polygons
-	//Append Polygons
-	for (const FPolygonID SourcePolygonID : InSource.Polygons().GetElementIDs())
-	{
-		TArray<FVertexInstanceID> SourceVertexInstanceIDs = InSource.GetPolygonVertexInstances(SourcePolygonID);
-		const FPolygonGroupID PolygonGroupID = InSource.GetPolygonPolygonGroup(SourcePolygonID);
-
-		TArray<FVertexInstanceID> ContourVertexInstances;
-		CreateContour(SourceVertexInstanceIDs, ContourVertexInstances);
-
-		// Insert a polygon into the mesh
-		const FPolygonID TargetPolygonID = InTarget.CreatePolygon(PolygonGroupID, ContourVertexInstances);
-	}
-}
-
 void FMeshMergeHelpers::MergeImpostersToMesh(TArray<const UStaticMeshComponent*> ImposterComponents, FMeshDescription& InMeshDescription, const FVector& InPivot, int32 InBaseMaterialIndex, TArray<UMaterialInterface*>& OutImposterMaterials)
 {
 	TMap<UMaterialInterface*, FPolygonGroupID> ImposterMaterialToPolygonGroupID;
@@ -1341,6 +1211,12 @@ void FMeshMergeHelpers::MergeImpostersToMesh(TArray<const UStaticMeshComponent*>
 		}
 		ImposterMesh.RemapPolygonGroups(RemapSourcePolygonGroup);
 
-		AppendMesh(InMeshDescription, ImposterMesh);
+		FStaticMeshOperations::FAppendSettings AppendSettings;
+		AppendSettings.bMergeVertexColor = true;
+		for (int32 ChannelIdx = 0; ChannelIdx < FStaticMeshOperations::FAppendSettings::MAX_NUM_UV_CHANNELS; ++ChannelIdx)
+		{
+			AppendSettings.bMergeUVChannels[ChannelIdx] = true;
+		}
+		FStaticMeshOperations::AppendMeshDescription(ImposterMesh, InMeshDescription, AppendSettings);
 	}
 }
diff --git a/Engine/Source/Developer/MeshMergeUtilities/Private/MeshMergeHelpers.h b/Engine/Source/Developer/MeshMergeUtilities/Private/MeshMergeHelpers.h
--- a/Engine/Source/Developer/MeshMergeUtilities/Private/MeshMergeHelpers.h
+++ b/Engine/Source/Developer/MeshMergeUtilities/Private/MeshMergeHelpers.h
@@ -28,7 +28,7 @@ public:
 	static void ExtractSections(const UStaticMesh* StaticMesh, int32 LODIndex, TArray<FSectionInfo>& OutSections);
 
 	/** Expanding instance data from instanced static mesh components */
-	static void ExpandInstances(const UInstancedStaticMeshComponent* InInstancedStaticMeshComponent, FMeshDescription& InOutRawMesh, TArray<FSectionInfo>& InOutSections);
+	static void ExpandInstances(const UInstancedStaticMeshComponent* InInstancedStaticMeshComponent, FMeshDescription& InOutRawMesh);
 
 	/** Extracting mesh data in FMeshDescription form from static, skeletal mesh (components) */
 	static void RetrieveMesh(const UStaticMeshComponent* StaticMeshComponent, int32 LODIndex, FMeshDescription& OutMeshDescription, bool bPropagateVertexColours);
@@ -72,9 +72,6 @@ public:
 	/** Checks whether or not the landscape proxy is hit given a ray start and end */
 	static bool IsLandscapeHit(const FVector& RayOrigin, const FVector& RayEndPoint, const UWorld* World, const TArray<ALandscapeProxy*>& LandscapeProxies, FVector& OutHitLocation);
 	
-	/** Appends a FMeshDescription to another instance */
-	static void AppendMesh(FMeshDescription& InTarget, const FMeshDescription& InSource);
-
 	/** Merges imposter meshes into the given MeshDescription. */
 	static void MergeImpostersToMesh(TArray<const UStaticMeshComponent*> ImposterComponents, FMeshDescription& InOutMeshDescription, const FVector& InPivot, int32 BaseMaterialIndex, TArray<UMaterialInterface*>& OutImposterMaterials);
 };
diff --git a/Engine/Source/Developer/MeshMergeUtilities/Private/MeshMergeUtilities.cpp b/Engine/Source/Developer/MeshMergeUtilities/Private/MeshMergeUtilities.cpp
--- a/Engine/Source/Developer/MeshMergeUtilities/Private/MeshMergeUtilities.cpp
+++ b/Engine/Source/Developer/MeshMergeUtilities/Private/MeshMergeUtilities.cpp
@@ -74,6 +74,8 @@
 #include "Async/Async.h"
 #include "TextureCompiler.h"
 
+#include "ISMPartition/ISMComponentDescriptor.h"
+
 #if WITH_EDITOR
 #include "Widgets/Notifications/SNotificationList.h"
 #include "Framework/Notifications/NotificationManager.h"
@@ -185,7 +187,7 @@ void FMeshMergeUtilities::BakeMaterialsForComponent(TArray<TWeakObjectPtr<UObjec
 			for (const int32 LODIndex : MaterialOptions->LODIndices)
 			{
 				FMeshData MeshSettings;
-				MeshSettings.RawMeshDescription = nullptr;
+				MeshSettings.MeshDescription = nullptr;
 
 				// Add material indices used for rendering out material
 				for (const auto& Pair : UniqueSectionIndexPerLOD[LODIndex])
@@ -199,13 +201,13 @@ void FMeshMergeUtilities::BakeMaterialsForComponent(TArray<TWeakObjectPtr<UObjec
 				if (MeshSettings.MaterialIndices.Num())
 				{
 					// Retrieve raw mesh
-					MeshSettings.RawMeshDescription = RawMeshLODs.Find(LODIndex);
+					MeshSettings.MeshDescription = RawMeshLODs.Find(LODIndex);
 					
 					//Should not be using mesh data if there is no mesh
-					check(MeshSettings.RawMeshDescription);
+					check(MeshSettings.MeshDescription);
 
 					MeshSettings.TextureCoordinateBox = FBox2D(FVector2D(0.0f, 0.0f), FVector2D(1.0f, 1.0f));
-					const bool bUseVertexColor = FStaticMeshOperations::HasVertexColor(*(MeshSettings.RawMeshDescription));
+					const bool bUseVertexColor = FStaticMeshOperations::HasVertexColor(*(MeshSettings.MeshDescription));
 					if (MaterialOptions->bUseSpecificUVIndex)
 					{
 						MeshSettings.TextureCoordinateIndex = MaterialOptions->TextureCoordinateIndex;
@@ -223,10 +225,10 @@ void FMeshMergeUtilities::BakeMaterialsForComponent(TArray<TWeakObjectPtr<UObjec
 					Adapter->ApplySettings(LODIndex, MeshSettings);
 					
 					// In case part of the UVs is not within the 0-1 range try to use the lightmap UVs
-					const bool bNeedsUniqueUVs = FMeshMergeHelpers::CheckWrappingUVs(*(MeshSettings.RawMeshDescription), MeshSettings.TextureCoordinateIndex);
+					const bool bNeedsUniqueUVs = FMeshMergeHelpers::CheckWrappingUVs(*(MeshSettings.MeshDescription), MeshSettings.TextureCoordinateIndex);
 					const int32 LightMapUVIndex = Adapter->LightmapUVIndex();
 					
-					TVertexInstanceAttributesConstRef<FVector2D> VertexInstanceUVs = FStaticMeshAttributes(*MeshSettings.RawMeshDescription).GetVertexInstanceUVs();
+					TVertexInstanceAttributesConstRef<FVector2D> VertexInstanceUVs = FStaticMeshConstAttributes(*MeshSettings.MeshDescription).GetVertexInstanceUVs();
 					if (bNeedsUniqueUVs && MeshSettings.TextureCoordinateIndex != LightMapUVIndex && VertexInstanceUVs.GetNumElements() > 0 && VertexInstanceUVs.GetNumChannels() > LightMapUVIndex)
 					{
 						MeshSettings.TextureCoordinateIndex = LightMapUVIndex;
@@ -276,7 +278,7 @@ void FMeshMergeUtilities::BakeMaterialsForComponent(TArray<TWeakObjectPtr<UObjec
 
 			if (MeshSettings.MaterialIndices.Num())
 			{
-				MeshSettings.RawMeshDescription = nullptr;
+				MeshSettings.MeshDescription = nullptr;
 				MeshSettings.TextureCoordinateBox = FBox2D(FVector2D(0.0f, 0.0f), FVector2D(1.0f, 1.0f));
 				MeshSettings.TextureCoordinateIndex = 0;
 
@@ -1100,6 +1102,289 @@ float FMeshMergeUtilities::FlattenEmissivescale(TArray<struct FFlattenMaterial>&
 	return MaxScale;
 }
 
+class FProxyMeshDescriptor
+{
+public:
+	FProxyMeshDescriptor(const UStaticMeshComponent* StaticMeshComponent, int32 LODIndex)
+		: LODIndex(LODIndex)
+		, LightMapIndex(INDEX_NONE)
+	{
+		ISMDescriptor.InitFrom(StaticMeshComponent, false);
+		ISMDescriptor.ComputeHash();
+
+		// Retrieve lightmap for usage of lightmap data
+		if (StaticMeshComponent->LODData.IsValidIndex(0))
+		{
+			const FStaticMeshComponentLODInfo& ComponentLODInfo = StaticMeshComponent->LODData[0];
+			const FMeshMapBuildData* MeshMapBuildData = StaticMeshComponent->GetMeshMapBuildData(ComponentLODInfo);
+			if (MeshMapBuildData)
+			{
+				LightMap = MeshMapBuildData->LightMap;
+				LightMapIndex = StaticMeshComponent->GetStaticMesh()->GetLightMapCoordinateIndex();
+			}
+		}
+
+		
+		Hash = ISMDescriptor.ComputeHash();
+
+		FCrc::TypeCrc32(LODIndex, Hash);
+
+		if (LightMapIndex != INDEX_NONE)
+		{
+			FCrc::TypeCrc32(LightMap.GetReference(), Hash);
+			FCrc::TypeCrc32(LightMapIndex, Hash);
+		}
+	}
+
+	bool operator==(const FProxyMeshDescriptor& InOther) const
+	{
+		return Hash == InOther.Hash &&
+			   LODIndex == InOther.LODIndex &&
+			   LightMap == InOther.LightMap &&
+			   LightMapIndex == InOther.LightMapIndex &&
+			   ISMDescriptor == InOther.ISMDescriptor;
+	}
+
+	int32 GetLODIndex() const { return LODIndex; }
+	FLightMapRef GetLightMap() const { return LightMap; }
+	int32 GetLightMapIndex() const { return LightMapIndex; }
+
+	UStaticMesh* GetStaticMesh() const { return ISMDescriptor.StaticMesh; }
+
+	const FMeshDescription& GetMeshDescription() const
+	{
+		return MeshDescription;
+	}
+
+	void PrepareMeshDescription(const FMeshProxySettings& InMeshProxySettings)
+	{
+		TRACE_CPUPROFILER_EVENT_SCOPE(PrepareMeshDescription);
+
+		// Retrieve mesh data in FMeshDescription form
+		FStaticMeshAttributes(MeshDescription).Register();
+		FMeshMergeHelpers::RetrieveMesh(ISMDescriptor.StaticMesh, LODIndex, MeshDescription);
+
+		TArray<FVector2D> TexCoords;
+		bool bSuccess = FStaticMeshOperations::GenerateUniqueUVsForStaticMesh(MeshDescription, InMeshProxySettings.MaterialSettings.TextureSize.GetMax(), false, TexCoords);
+
+		if (bSuccess)
+		{
+			TVertexInstanceAttributesRef<FVector2D> VertexInstanceUVs = FStaticMeshAttributes(MeshDescription).GetVertexInstanceUVs();
+			for (int32 Idx = 0; Idx < TexCoords.Num(); Idx++)
+			{
+				VertexInstanceUVs.Set(Idx, TexCoords[Idx]);
+			}
+		}
+		else
+		{
+			UE_LOG(LogMeshMerging, Warning, TEXT("GenerateUniqueUVsForStaticMesh: Failed to pack UVs for static mesh \"%s\" (num triangles = %d, texture resolution = %d)."), *ISMDescriptor.StaticMesh->GetName(), MeshDescription.Triangles().Num(), InMeshProxySettings.MaterialSettings.TextureSize.GetMax());
+			MeshDescription.Empty();
+		}
+	}
+
+private:
+	int32			Hash;
+
+	int32			LODIndex;
+	FLightMapRef	LightMap;
+	int32			LightMapIndex;
+
+	FISMComponentDescriptor ISMDescriptor;
+
+	FMeshDescription MeshDescription;
+};
+
+static void ScaleTextureCoordinatesToBox(const FBox2D& Box, TArray<FVector2D>& InOutTextureCoordinates)
+{
+	const FBox2D CoordinateBox(InOutTextureCoordinates);
+	const FVector2D CoordinateRange = CoordinateBox.GetSize();
+	const FVector2D Offset = CoordinateBox.Min + Box.Min;
+	const FVector2D Scale = Box.GetSize() / CoordinateRange;
+	for (FVector2D& Coordinate : InOutTextureCoordinates)
+	{
+		Coordinate = (Coordinate - Offset) * Scale;
+	}
+}
+
+typedef TFunctionRef<int32(const UStaticMeshComponent*)> FGetMeshLODFunc;
+
+static TArray<FMeshDescription*> GatherGeometry(const TArray<UStaticMeshComponent*>& InStaticMeshComponents, const FMeshProxySettings& InMeshProxySettings, TArray<FProxyMeshDescriptor>& InDescriptors, const TArray<TArray<int32>>& InMeshesToMergePerDescriptor, const TArray<int32>& InMeshesToMeshDescriptor, FGetMeshLODFunc InGetMeshLODFunc, int32& OutSummedLightmapPixels)
+{
+	TRACE_CPUPROFILER_EVENT_SCOPE(FMeshMergeUtilities::GatherGeometry);
+
+	TArray<FMeshDescription> TempDescriptionData;
+	TempDescriptionData.SetNum(InStaticMeshComponents.Num());
+
+	TAtomic<int32>  SummedLightmapPixels(0);
+
+	TArray<FMeshDescription*> MeshesDescriptions;
+	MeshesDescriptions.SetNum(InDescriptors.Num());
+
+	// If grouping by identical meshes, prepare each mesh description along with their flattened UVs
+	// These meshes descriptions will serve for material baking, but also as the basis for creating a
+	// single merged mesh out of all instances.
+	if (InMeshProxySettings.bGroupIdenticalMeshesForBaking)
+	{
+		TRACE_CPUPROFILER_EVENT_SCOPE(UniqueUVs);
+		ParallelFor(InDescriptors.Num(), [&InDescriptors, &InMeshProxySettings](uint32 Index)
+		{
+			InDescriptors[Index].PrepareMeshDescription(InMeshProxySettings);
+		});
+	}
+
+	// Gather geometry from each component, expand ISMC geometry
+	{
+		TRACE_CPUPROFILER_EVENT_SCOPE(GatherGeometryFromComponents);
+		ParallelFor(InStaticMeshComponents.Num(), [InStaticMeshComponents, &InGetMeshLODFunc, InDescriptors, InMeshesToMeshDescriptor, InMeshProxySettings, &TempDescriptionData, &SummedLightmapPixels](uint32 Index)
+		{
+			TRACE_CPUPROFILER_EVENT_SCOPE(GatherGeometryFromComponent);
+			const UStaticMeshComponent* StaticMeshComponent = InStaticMeshComponents[Index];
+
+			FMeshDescription& MeshDescription = TempDescriptionData[Index];
+
+			// Retrieve meshes
+			if (!InMeshProxySettings.bGroupIdenticalMeshesForBaking || InDescriptors[InMeshesToMeshDescriptor[Index]].GetMeshDescription().IsEmpty())
+			{
+				const int32 LODIndex = InGetMeshLODFunc(StaticMeshComponent);
+				static const bool bPropagateVertexColours = true;
+
+				// Retrieve mesh data in FMeshDescription form
+				FStaticMeshAttributes(MeshDescription).Register();
+				FMeshMergeHelpers::RetrieveMesh(StaticMeshComponent, LODIndex, MeshDescription, bPropagateVertexColours);
+			}
+			else
+			{
+				MeshDescription = InDescriptors[InMeshesToMeshDescriptor[Index]].GetMeshDescription();
+				FStaticMeshOperations::ApplyTransform(MeshDescription, StaticMeshComponent->GetComponentTransform());
+			}
+
+			// If the component is an ISMC then we need to duplicate the vertex data
+			int32 NumInstances = 1;
+			if (const UInstancedStaticMeshComponent* InstancedStaticMeshComponent = Cast<UInstancedStaticMeshComponent>(StaticMeshComponent))
+			{
+				FMeshMergeHelpers::ExpandInstances(InstancedStaticMeshComponent, MeshDescription);
+				NumInstances = InstancedStaticMeshComponent->PerInstanceSMData.Num();
+			}
+
+			int32 LightMapWidth, LightMapHeight;
+			StaticMeshComponent->GetLightMapResolution(LightMapWidth, LightMapHeight);
+			// Make sure we at least have some lightmap space allocated in case the static mesh is set up with invalid input
+			SummedLightmapPixels += FMath::Max(16, LightMapHeight * LightMapWidth * NumInstances);
+		}, EParallelForFlags::Unbalanced);
+	}
+
+	// For each mesh, append each component geometry
+	{
+		TRACE_CPUPROFILER_EVENT_SCOPE(AppendMeshes);
+		FStaticMeshOperations::FAppendSettings AppendSettings;
+		for (int32 ChannelIdx = 0; ChannelIdx < FStaticMeshOperations::FAppendSettings::MAX_NUM_UV_CHANNELS; ++ChannelIdx)
+		{
+			AppendSettings.bMergeUVChannels[ChannelIdx] = true;
+		}
+
+		ParallelFor(InDescriptors.Num(), [&MeshesDescriptions, &InMeshesToMergePerDescriptor, &TempDescriptionData, &AppendSettings, &InMeshProxySettings](uint32 Index)
+		{
+			TRACE_CPUPROFILER_EVENT_SCOPE(AppendMeshes);
+			FMeshDescription* TargetMeshDescription = new FMeshDescription();
+			FStaticMeshAttributes(*TargetMeshDescription).Register();
+
+			TArray<const FMeshDescription*> SourceMeshDescriptions;
+			SourceMeshDescriptions.Reserve(InMeshesToMergePerDescriptor[Index].Num());
+			for (int32 TempIdx : InMeshesToMergePerDescriptor[Index])
+			{
+				SourceMeshDescriptions.Add(&TempDescriptionData[TempIdx]);
+			}
+
+			FStaticMeshOperations::AppendMeshDescriptions(SourceMeshDescriptions, *TargetMeshDescription, AppendSettings);
+
+			MeshesDescriptions[Index] = TargetMeshDescription;
+		}, EParallelForFlags::Unbalanced);
+	}
+
+	OutSummedLightmapPixels = SummedLightmapPixels;
+
+	return MeshesDescriptions;
+}
+
+TArray<FMeshData> PrepareBakingMeshes(const struct FMeshProxySettings& InMeshProxySettings, const TArray<FProxyMeshDescriptor>& InDescriptors, TArray<FMeshDescription*> InMeshDescriptions)
+{
+	TRACE_CPUPROFILER_EVENT_SCOPE(PrepareBakingMeshes)
+
+	check(InDescriptors.Num() == InMeshDescriptions.Num());
+
+	TArray<FMeshData> MeshData;
+	MeshData.SetNum(InDescriptors.Num());
+
+	ParallelFor(InDescriptors.Num(), [&MeshData, &InDescriptors, &InMeshDescriptions, &InMeshProxySettings](uint32 MeshIndex)
+	{
+		const FProxyMeshDescriptor& MeshDescriptor = InDescriptors[MeshIndex];
+
+		FMeshData& MeshSettings = MeshData[MeshIndex];
+
+		bool bSuccess = true;
+
+		if (InMeshProxySettings.bGroupIdenticalMeshesForBaking)
+		{
+			// Grouping by identical meshes, the UVs should have already been setup
+			bSuccess = !MeshDescriptor.GetMeshDescription().IsEmpty();
+			MeshSettings.MeshDescription = &MeshDescriptor.GetMeshDescription();
+		}
+		else
+		{
+			FMeshDescription& MeshDescription = *InMeshDescriptions[MeshIndex];
+			MeshSettings.MeshDescription = &MeshDescription;
+
+			if (MeshDescriptor.GetLightMapIndex() != INDEX_NONE)
+			{
+				MeshSettings.LightMap = MeshDescriptor.GetLightMap();
+				MeshSettings.LightMapIndex = MeshDescriptor.GetLightMapIndex();
+			}
+
+			TVertexInstanceAttributesRef<FVector2D> VertexInstanceUVs = FStaticMeshAttributes(MeshDescription).GetVertexInstanceUVs();
+
+		    // If we already have lightmap uvs generated and they are valid, we can reuse those instead of having to generate new ones
+		    const int32 LightMapCoordinateIndex = MeshDescriptor.GetStaticMesh()->GetLightMapCoordinateIndex();
+		    if (InMeshProxySettings.bReuseMeshLightmapUVs &&
+			    LightMapCoordinateIndex > 0 &&
+			    VertexInstanceUVs.GetNumElements() > 0 &&
+			    VertexInstanceUVs.GetNumChannels() > LightMapCoordinateIndex)
+		    {
+			    MeshSettings.CustomTextureCoordinates.Reset(VertexInstanceUVs.GetNumElements());
+			    for (const FVertexInstanceID VertexInstanceID : MeshDescription.VertexInstances().GetElementIDs())
+			    {
+				    MeshSettings.CustomTextureCoordinates.Add(VertexInstanceUVs.Get(VertexInstanceID, LightMapCoordinateIndex));
+			    }
+			    MeshSettings.TextureCoordinateBox = FBox2D(FVector2D::ZeroVector, FVector2D(1, 1));
+			    ScaleTextureCoordinatesToBox(MeshSettings.TextureCoordinateBox, MeshSettings.CustomTextureCoordinates);
+		    }
+		    else
+		    {
+			    // Generate unique UVs for mesh (should only be done if needed)
+			    bSuccess = FStaticMeshOperations::GenerateUniqueUVsForStaticMesh(MeshDescription, InMeshProxySettings.MaterialSettings.TextureSize.GetMax(), false, MeshSettings.CustomTextureCoordinates);
+			    if (bSuccess)
+			    {
+				    MeshSettings.TextureCoordinateBox = FBox2D(FVector2D::ZeroVector, FVector2D(1, 1));
+				    ScaleTextureCoordinatesToBox(MeshSettings.TextureCoordinateBox, MeshSettings.CustomTextureCoordinates);
+			    }
+			    else
+			    {
+				    UE_LOG(LogMeshMerging, Warning, TEXT("GenerateUniqueUVsForStaticMesh: Failed to pack UVs for static mesh \"%s\" (num triangles = %d, texture resolution = %d)."), *MeshDescriptor.GetStaticMesh()->GetName(), MeshDescription.Triangles().Num(), InMeshProxySettings.MaterialSettings.TextureSize.GetMax());
+			    }
+		    }
+		}
+
+		if (!bSuccess)
+		{
+			MeshSettings.MeshDescription = nullptr;
+			MeshSettings.TextureCoordinateBox = FBox2D(FVector2D(0.0f, 0.0f), FVector2D(1.0f, 1.0f));
+			MeshSettings.TextureCoordinateIndex = 0;
+			MeshSettings.CustomTextureCoordinates.Empty();
+		}
+	});
+
+	return MeshData;
+}
+
 void FMeshMergeUtilities::CreateProxyMesh(const TArray<AActor*>& InActors, const struct FMeshProxySettings& InMeshProxySettings, UPackage* InOuter, const FString& InProxyBasePackageName, const FGuid InGuid, const FCreateProxyDelegate& InProxyCreatedDelegate, const bool bAllowAsync, const float ScreenSize) const
 {
 	UMaterial* BaseMaterial = LoadObject<UMaterial>(NULL, TEXT("/Engine/EngineMaterials/BaseFlattenMaterial.BaseFlattenMaterial"), NULL, LOAD_None, NULL);
@@ -1186,7 +1471,6 @@ void FMeshMergeUtilities::CreateProxyMesh(const TArray<UStaticMeshComponent*>& I
 
 	TArray<FRawMeshExt> SourceMeshes;
 	TMap<FMeshIdAndLOD, TArray<int32>> GlobalMaterialMap;
-	static const int32 ProxyMeshTargetLODLevel = 0;
 
 	FBoxSphereBounds EstimatedBounds(ForceInitToZero);
 	for (const UStaticMeshComponent* StaticMeshComponent : ComponentsToMerge)
@@ -1195,12 +1479,21 @@ void FMeshMergeUtilities::CreateProxyMesh(const TArray<UStaticMeshComponent*>& I
 	}
 
 	static const float FOVRad = FMath::DegreesToRadians(45.0f);
-
 	static const FMatrix ProjectionMatrix = FPerspectiveMatrix(FOVRad, 1920, 1080, 0.01f);
 	FHierarchicalLODUtilitiesModule& HLODModule = FModuleManager::LoadModuleChecked<FHierarchicalLODUtilitiesModule>("HierarchicalLODUtilities");
 	IHierarchicalLODUtilities* Utilities = HLODModule.GetUtilities();
 	float EstimatedDistance = Utilities->CalculateDrawDistanceFromScreenSize(EstimatedBounds.SphereRadius, ScreenSize, ProjectionMatrix);
 
+	auto SelectLODFunc = [&InMeshProxySettings, &Utilities, EstimatedDistance] (const UStaticMeshComponent* Component)
+	{
+		int32 LODIndex = 0;
+		if (InMeshProxySettings.bCalculateCorrectLODModel)
+		{
+			LODIndex = Utilities->GetLODLevelForScreenSize(Component, Utilities->CalculateScreenSizeFromDrawDistance(Component->Bounds.SphereRadius, ProjectionMatrix, EstimatedDistance));
+		}
+		return LODIndex;
+	};
+
 	SlowTask.EnterProgressFrame(5.0f, LOCTEXT("CreateProxyMesh_CollectingMeshes", "Collecting Input Static Meshes"));
 
 	// Mesh / LOD index	
@@ -1213,105 +1506,70 @@ void FMeshMergeUtilities::CreateProxyMesh(const TArray<UStaticMeshComponent*>& I
 	TArray<FSectionInfo> UniqueSections;
 	TMultiMap<uint32, uint32> SectionToMesh;
 
-	// Copies of mesh data
-	TArray<FMeshDescription*> MeshDescriptionData;
-	MeshDescriptionData.SetNum(ComponentsToMerge.Num());
-
 	TArray<const UStaticMeshComponent*> ImposterMeshComponents;
-	ImposterMeshComponents.SetNum(ComponentsToMerge.Num());
-
 	TArray<UStaticMeshComponent*> StaticMeshComponents;
-	StaticMeshComponents.SetNum(ComponentsToMerge.Num());
+	for (UStaticMeshComponent* StaticMeshComponent : ComponentsToMerge)
+	{
+		if (StaticMeshComponent->bUseMaxLODAsImposter)
+		{
+			ImposterMeshComponents.Add(StaticMeshComponent);
+		}
+		else
+		{
+			StaticMeshComponents.Add(StaticMeshComponent);
+		}
+	}
 
-	TAtomic<int32>  SummedLightmapPixels(0);
-	TAtomic<uint32> ImposterMeshComponentsIndex(0);
-	TAtomic<uint32> StaticMeshComponentsIndex(0);
-
-	{
-		TRACE_CPUPROFILER_EVENT_SCOPE(FMeshMergeUtilities::MeshGathering);
-
-		TArray<TArray<FSectionInfo>> GlobalSections;
-		GlobalSections.SetNum(ComponentsToMerge.Num());
-
-		ParallelFor(
-			ComponentsToMerge.Num(),
-			[
-				&ComponentsToMerge,
-				&ImposterMeshComponents,
-				&StaticMeshComponents,
-				&ImposterMeshComponentsIndex,
-				&StaticMeshComponentsIndex,
-				&Utilities,
-				&EstimatedDistance,
-				&InMeshProxySettings,
-				&MeshDescriptionData,
-				&GlobalSections,
-				&SummedLightmapPixels
-			](uint32 Index)
-			{
-				UStaticMeshComponent* StaticMeshComponent = ComponentsToMerge[Index];
-
-				int32 NumInstances = 1;
-				if (StaticMeshComponent->bUseMaxLODAsImposter)
-				{
-					ImposterMeshComponents[ImposterMeshComponentsIndex++] = StaticMeshComponent;
-				}
-				else
-				{
-					const int32 MeshIndex = StaticMeshComponentsIndex++;
-					StaticMeshComponents[MeshIndex] = StaticMeshComponent;
+	TArray<FProxyMeshDescriptor>	MeshDescriptors;
+	TArray<TArray<int32>>			MeshesToMergePerDescriptor;
+	TArray<int32>					MeshToMeshDescriptor;
 
-					const int32 ScreenSizeBasedLODLevel = Utilities->GetLODLevelForScreenSize(StaticMeshComponent, Utilities->CalculateScreenSizeFromDrawDistance(StaticMeshComponent->Bounds.SphereRadius, ProjectionMatrix, EstimatedDistance));
-					const int32 LODIndex = InMeshProxySettings.bCalculateCorrectLODModel ? ScreenSizeBasedLODLevel : 0;
-					static const bool bPropagateVertexColours = true;
+	TArray<TArray<FSectionInfo>>	GlobalSections;
 
-					// Retrieve mesh data in FMeshDescription form
-					MeshDescriptionData[MeshIndex] = new FMeshDescription();
-					FMeshDescription& MeshDescription = *MeshDescriptionData[MeshIndex];
-					FStaticMeshAttributes(MeshDescription).Register();
-					FMeshMergeHelpers::RetrieveMesh(StaticMeshComponent, LODIndex, MeshDescription, bPropagateVertexColours);
+	MeshToMeshDescriptor.Reserve(StaticMeshComponents.Num());
 
-					TArray<FSectionInfo>& Sections = GlobalSections[MeshIndex];
+	for (int32 ComponentIndex = 0; ComponentIndex < StaticMeshComponents.Num(); ++ComponentIndex)
+	{
+		const UStaticMeshComponent* StaticMeshComponent = StaticMeshComponents[ComponentIndex];
 
-					// Extract sections for given LOD index from the mesh 
-					FMeshMergeHelpers::ExtractSections(StaticMeshComponent, LODIndex, Sections);
+		FProxyMeshDescriptor MeshDescriptor(StaticMeshComponent, SelectLODFunc(StaticMeshComponent));
 
-					// If the component is an ISMC then we need to duplicate the vertex data
-					if(StaticMeshComponent->IsA<UInstancedStaticMeshComponent>())
-					{
-						const UInstancedStaticMeshComponent* InstancedStaticMeshComponent = Cast<UInstancedStaticMeshComponent>(StaticMeshComponent);
-						FMeshMergeHelpers::ExpandInstances(InstancedStaticMeshComponent, MeshDescription, Sections);
-						NumInstances = InstancedStaticMeshComponent->PerInstanceSMData.Num();
-					}
-				}	
+		int32 Index;
+		if (!InMeshProxySettings.bGroupIdenticalMeshesForBaking || !MeshDescriptors.Find(MeshDescriptor, Index))
+		{
+			Index = MeshDescriptors.Num();
+			MeshDescriptors.Add(MeshDescriptor);
+			MeshesToMergePerDescriptor.AddDefaulted();
 
-				int32 LightMapWidth, LightMapHeight;
-				StaticMeshComponent->GetLightMapResolution(LightMapWidth, LightMapHeight);
-				// Make sure we at least have some lightmap space allocated in case the static mesh is set up with invalid input
-				SummedLightmapPixels += FMath::Max(16, LightMapHeight * LightMapWidth * NumInstances);
-			},
-			EParallelForFlags::Unbalanced
-		);
+			TArray<FSectionInfo>& Sections = GlobalSections.AddDefaulted_GetRef();
 
-		ImposterMeshComponents.SetNum(ImposterMeshComponentsIndex);
-		StaticMeshComponents.SetNum(StaticMeshComponentsIndex);
-		MeshDescriptionData.SetNum(StaticMeshComponentsIndex);
+			// Extract sections for given LOD index from the mesh 
+			FMeshMergeHelpers::ExtractSections(StaticMeshComponent, MeshDescriptor.GetLODIndex(), Sections);
+		}
 
-		for (uint32 MeshIndex = 0; MeshIndex < StaticMeshComponentsIndex; ++MeshIndex)
-		{
-			TArray<FSectionInfo>& Sections = GlobalSections[MeshIndex];
+		MeshesToMergePerDescriptor[Index].Add(ComponentIndex);
+		MeshToMeshDescriptor.Add(Index);
+	}
 
-			for (int32 SectionIndex = 0; SectionIndex < Sections.Num(); ++SectionIndex)
-			{
-				FSectionInfo& Section = Sections[SectionIndex];
+	for (int32 MeshIndex = 0; MeshIndex < GlobalSections.Num(); ++MeshIndex)
+	{
+		TArray<FSectionInfo>& Sections = GlobalSections[MeshIndex];
 
-				const int32 UniqueIndex = UniqueSections.AddUnique(Section);
-				MeshSectionToUniqueSection.Add(MeshIndex, TPair<uint32, uint32>(SectionIndex, UniqueIndex));
-				SectionToMesh.Add(UniqueIndex, MeshIndex);
-			}
+		for (int32 SectionIndex = 0; SectionIndex < Sections.Num(); ++SectionIndex)
+		{
+			FSectionInfo& Section = Sections[SectionIndex];
+
+			const int32 UniqueIndex = UniqueSections.AddUnique(Section);
+			MeshSectionToUniqueSection.Add(MeshIndex, TPair<uint32, uint32>(SectionIndex, UniqueIndex));
+			SectionToMesh.Add(UniqueIndex, MeshIndex);
 		}
 	}
 
+	int32 SummedLightmapPixels;
+	
+	TArray<FMeshDescription*> MeshDescriptionData = GatherGeometry(StaticMeshComponents, InMeshProxySettings, MeshDescriptors, MeshesToMergePerDescriptor, MeshToMeshDescriptor, SelectLODFunc, SummedLightmapPixels);
+	TArray<FMeshData> MeshBakingData = PrepareBakingMeshes(InMeshProxySettings, MeshDescriptors, MeshDescriptionData);
+
 	TArray<UMaterialInterface*> UniqueMaterials;
 	//Unique material index to unique section index
 	TMultiMap<uint32, uint32> MaterialToSectionMap;
@@ -1340,7 +1598,6 @@ void FMeshMergeUtilities::CreateProxyMesh(const TArray<UStaticMeshComponent*>& I
 	{
 		TRACE_CPUPROFILER_EVENT_SCOPE(FMeshMergeUtilities::MaterialAnalysisAndUVGathering);
 
-		TArray<TFunction<void ()>> Lambdas;
 		for (int32 MaterialIndex = 0; MaterialIndex < UniqueMaterials.Num(); ++MaterialIndex)
 		{
 			UMaterialInterface* Material = UniqueMaterials[MaterialIndex];
@@ -1374,46 +1631,7 @@ void FMeshMergeUtilities::CreateProxyMesh(const TArray<UStaticMeshComponent*>& I
 
 					for (const uint32 MeshIndex : MeshIndices)
 					{
-						// Retrieve mesh description
-						const UStaticMeshComponent* StaticMeshComponent = StaticMeshComponents[MeshIndex];
-						FMeshDescription* MeshDescription = MeshDescriptionData[MeshIndex];
-
-						FMeshData MeshSettings;
-						MeshSettings.RawMeshDescription = MeshDescription;
-
-						TVertexInstanceAttributesRef<FVector2D> VertexInstanceUVs = FStaticMeshAttributes(*MeshDescription).GetVertexInstanceUVs();
-
-						// If we already have lightmap uvs generated and they are valid, we can reuse those instead of having to generate new ones
-						const int32 LightMapCoordinateIndex = StaticMeshComponent->GetStaticMesh()->GetLightMapCoordinateIndex();
-						if (InMeshProxySettings.bReuseMeshLightmapUVs &&
-							LightMapCoordinateIndex > 0 &&
-							VertexInstanceUVs.GetNumElements() > 0 &&
-							VertexInstanceUVs.GetNumChannels() > LightMapCoordinateIndex)
-						{
-							MeshSettings.CustomTextureCoordinates.Reset(VertexInstanceUVs.GetNumElements());
-							for (const FVertexInstanceID VertexInstanceID : MeshDescription->VertexInstances().GetElementIDs())
-							{
-								MeshSettings.CustomTextureCoordinates.Add(VertexInstanceUVs.Get(VertexInstanceID, LightMapCoordinateIndex));
-							}
-							ScaleTextureCoordinatesToBox(FBox2D(FVector2D::ZeroVector, FVector2D(1, 1)), MeshSettings.CustomTextureCoordinates);
-						}
-						else
-						{
-							// Accumulate slow running tasks to process them in parallel once the arrays
-							// are finished being resized.
-							Lambdas.Emplace(
-								[this, GlobalMeshSettingsIndex = GlobalMeshSettings.Num(), &GlobalMeshSettings, MeshDescription, Options]()
-								{
-									FMeshData& MeshSettings = GlobalMeshSettings[GlobalMeshSettingsIndex];
-									// Generate unique UVs for mesh (should only be done if needed)
-									FStaticMeshOperations::GenerateUniqueUVsForStaticMesh(*MeshDescription, Options->TextureSize.GetMax(), false, MeshSettings.CustomTextureCoordinates);
-									ScaleTextureCoordinatesToBox(FBox2D(FVector2D::ZeroVector, FVector2D(1, 1)), MeshSettings.CustomTextureCoordinates);
-									MeshSettings.TextureCoordinateBox = FBox2D(MeshSettings.CustomTextureCoordinates);
-								}
-							);
-						}
-						
-						MeshSettings.TextureCoordinateBox = FBox2D(MeshSettings.CustomTextureCoordinates);
+						FMeshData MeshSettings = MeshBakingData[MeshIndex];
 
 						// Section index is a unique one so we need to map it to the mesh's equivalent(s)
 						TArray<TPair<uint32, uint32>> SectionToUniqueSectionIndices;
@@ -1427,19 +1645,7 @@ void FMeshMergeUtilities::CreateProxyMesh(const TArray<UStaticMeshComponent*>& I
 							}
 						}
 
-						// Retrieve lightmap for usage of lightmap data
-						if (StaticMeshComponent->LODData.IsValidIndex(0))
-						{
-							const FStaticMeshComponentLODInfo& ComponentLODInfo = StaticMeshComponent->LODData[0];
-							const FMeshMapBuildData* MeshMapBuildData = StaticMeshComponent->GetMeshMapBuildData(ComponentLODInfo);
-							if (MeshMapBuildData)
-							{
-								MeshSettings.LightMap = MeshMapBuildData->LightMap;
-								MeshSettings.LightMapIndex = StaticMeshComponent->GetStaticMesh()->GetLightMapCoordinateIndex();
-							}
-						}
-
-						GlobalMeshSettings.Add(MeshSettings);
+						GlobalMeshSettings.Add(MoveTemp(MeshSettings));
 						GlobalMaterialSettings.Add(MaterialSettings);
 					}
 				}
@@ -1448,7 +1654,7 @@ void FMeshMergeUtilities::CreateProxyMesh(const TArray<UStaticMeshComponent*>& I
 			{
 				// Add simple bake entry 
 				FMeshData MeshSettings;
-				MeshSettings.RawMeshDescription = nullptr;
+				MeshSettings.MeshDescription = nullptr;
 				MeshSettings.TextureCoordinateBox = FBox2D(FVector2D(0.0f, 0.0f), FVector2D(1.0f, 1.0f));
 				MeshSettings.TextureCoordinateIndex = 0;
 
@@ -1476,19 +1682,6 @@ void FMeshMergeUtilities::CreateProxyMesh(const TArray<UStaticMeshComponent*>& I
 				GlobalMaterialSettings.Add(MaterialSettings);
 			}
 		}
-
-		if (Lambdas.Num())
-		{
-			TRACE_CPUPROFILER_EVENT_SCOPE(GenerateUVs);
-			ParallelFor(
-				Lambdas.Num(),
-				[&Lambdas](uint32 Index)
-				{
-					Lambdas[Index]();
-				},
-				EParallelForFlags::Unbalanced
-			);
-		}
 	}
 
 	TArray<FFlattenMaterial> FlattenedMaterials;
@@ -1601,16 +1794,14 @@ void FMeshMergeUtilities::CreateProxyMesh(const TArray<UStaticMeshComponent*>& I
 	// Add this proxy job to map	
 	Processor->AddProxyJob(InGuid, Data);
 
-	// We are only using LOD level 0 (ProxyMeshTargetLODLevel)
 	TArray<FMeshMergeData> MergeDataEntries;
-
 	{
 		TRACE_CPUPROFILER_EVENT_SCOPE(MergeDataPreparation)
 
 		for (int32 Index = 0; Index < MeshDescriptionData.Num(); ++Index)
 		{
 			FMeshMergeData MergeData;
-			MergeData.SourceStaticMesh = StaticMeshComponents[Index]->GetStaticMesh();
+			MergeData.SourceStaticMesh = MeshDescriptors[Index].GetStaticMesh();
 			MergeData.RawMesh = MeshDescriptionData[Index];
 			MergeData.bIsClippingMesh = false;
 
@@ -1618,7 +1809,7 @@ void FMeshMergeUtilities::CreateProxyMesh(const TArray<UStaticMeshComponent*>& I
 
 			FMeshData* MeshData = GlobalMeshSettings.FindByPredicate([&](const FMeshData& Entry)
 			{
-				return Entry.RawMeshDescription == MergeData.RawMesh && (Entry.CustomTextureCoordinates.Num() || Entry.TextureCoordinateIndex != 0);
+				return Entry.MeshDescription == MergeData.RawMesh && (Entry.CustomTextureCoordinates.Num() || Entry.TextureCoordinateIndex != 0);
 			});
 
 			if (MeshData)
@@ -1629,9 +1820,9 @@ void FMeshMergeUtilities::CreateProxyMesh(const TArray<UStaticMeshComponent*>& I
 				}
 				else
 				{
-					TVertexInstanceAttributesRef<FVector2D> VertexInstanceUVs = FStaticMeshAttributes(*MeshData->RawMeshDescription).GetVertexInstanceUVs();
-					MergeData.NewUVs.Reset(MeshData->RawMeshDescription->VertexInstances().Num());
-					for (const FVertexInstanceID VertexInstanceID : MeshData->RawMeshDescription->VertexInstances().GetElementIDs())
+					TVertexInstanceAttributesConstRef<FVector2D> VertexInstanceUVs = FStaticMeshConstAttributes(*MeshData->MeshDescription).GetVertexInstanceUVs();
+					MergeData.NewUVs.Reset(MeshData->MeshDescription->VertexInstances().Num());
+					for (const FVertexInstanceID VertexInstanceID : MeshData->MeshDescription->VertexInstances().GetElementIDs())
 					{
 						MergeData.NewUVs.Add(VertexInstanceUVs.Get(VertexInstanceID, MeshData->TextureCoordinateIndex));
 					}
@@ -1861,7 +2052,7 @@ bool RetrieveRawMeshData(FMeshMergeDataTracker& DataTracker
 	if (Component->IsA<UInstancedStaticMeshComponent>())
 	{
 		const UInstancedStaticMeshComponent* InstancedStaticMeshComponent = Cast<UInstancedStaticMeshComponent>(Component);
-		FMeshMergeHelpers::ExpandInstances(InstancedStaticMeshComponent, RawMesh, Sections);
+		FMeshMergeHelpers::ExpandInstances(InstancedStaticMeshComponent, RawMesh);
 	}
 
 	if (InSettings.bUseLandscapeCulling)
@@ -2189,17 +2380,17 @@ void FMeshMergeUtilities::MergeComponentsToStaticMesh(const TArray<UPrimitiveCom
 				
 				if (InSettings.bCreateMergedMaterial || bGloballyRemapUVs || (InSettings.bUseVertexDataForBakingMaterial && (bDoesMaterialUseVertexData || bRequiresUniqueUVs)))
 				{
-					FMeshDescription* RawMeshDescription = DataTracker.GetRawMeshPtr(Key);
-					MeshData.RawMeshDescription = RawMeshDescription;
+					FMeshDescription* MeshDescription = DataTracker.GetRawMeshPtr(Key);
+					MeshData.MeshDescription = MeshDescription;
 
 					// if it has vertex color/*WedgetColors.Num()*/, it should also use light map UV index
 					// we can't do this for all meshes, but only for the mesh that has vertex color.
-					if (bRequiresUniqueUVs || MeshData.RawMeshDescription->VertexInstances().Num() > 0)
+					if (bRequiresUniqueUVs || MeshData.MeshDescription->VertexInstances().Num() > 0)
 					{
 						// Check if there are lightmap uvs available?
 						const int32 LightMapUVIndex = StaticMeshComponentsToMerge[Key.GetMeshIndex()]->GetStaticMesh()->GetLightMapCoordinateIndex();
 
-						TVertexInstanceAttributesRef<FVector2D> VertexInstanceUVs = FStaticMeshAttributes(*MeshData.RawMeshDescription).GetVertexInstanceUVs();
+						TVertexInstanceAttributesConstRef<FVector2D> VertexInstanceUVs = FStaticMeshConstAttributes(*MeshData.MeshDescription).GetVertexInstanceUVs();
 						if (InSettings.bReuseMeshLightmapUVs && VertexInstanceUVs.GetNumElements() > 0 && VertexInstanceUVs.GetNumChannels() > LightMapUVIndex)
 						{
 							MeshData.TextureCoordinateIndex = LightMapUVIndex;
@@ -2211,10 +2402,10 @@ void FMeshMergeUtilities::MergeComponentsToStaticMesh(const TArray<UPrimitiveCom
 							if (!MeshLODsTextureCoordinates.Find(Tuple))
 							{
 								// No job found yet, fire an async task
-								MeshLODsTextureCoordinates.Add(Tuple, Async(EAsyncExecution::Thread, [RawMeshDescription, MaterialOptions, this]()
+								MeshLODsTextureCoordinates.Add(Tuple, Async(EAsyncExecution::Thread, [MeshDescription, MaterialOptions, this]()
 								{
 									TArray<FVector2D> UniqueTextureCoordinates;
-									FStaticMeshOperations::GenerateUniqueUVsForStaticMesh(*RawMeshDescription, MaterialOptions->TextureSize.GetMax(), false, UniqueTextureCoordinates);
+									FStaticMeshOperations::GenerateUniqueUVsForStaticMesh(*MeshDescription, MaterialOptions->TextureSize.GetMax(), false, UniqueTextureCoordinates);
 									ScaleTextureCoordinatesToBox(FBox2D(FVector2D::ZeroVector, FVector2D(1, 1)), UniqueTextureCoordinates);
 									return UniqueTextureCoordinates;
 								}));
@@ -2274,7 +2465,7 @@ void FMeshMergeUtilities::MergeComponentsToStaticMesh(const TArray<UPrimitiveCom
 				}
 				else
 				{
-					MeshData.RawMeshDescription = nullptr;
+					MeshData.MeshDescription = nullptr;
 					MeshData.TextureCoordinateBox = FBox2D(FVector2D(0.0f, 0.0f), FVector2D(1.0f, 1.0f));
 
 					// This prevents baking out the same material multiple times, which would be wasteful when it does not use vertex data anyway
@@ -2350,7 +2541,7 @@ void FMeshMergeUtilities::MergeComponentsToStaticMesh(const TArray<UPrimitiveCom
 				// copy to mesh data
 				for(FMeshData& MeshData : GlobalMeshSettings)
 				{
-					if(MeshData.RawMeshDescription == &RawMesh)
+					if(MeshData.MeshDescription == &RawMesh)
 					{
 						MeshData.CustomTextureCoordinates = UniqueTextureCoordinates;
 					}
@@ -3458,17 +3649,4 @@ void FMeshMergeUtilities::ExtractPhysicsDataFromComponents(const TArray<UPrimiti
 	}
 }
 
-void FMeshMergeUtilities::ScaleTextureCoordinatesToBox(const FBox2D& Box, TArray<FVector2D>& InOutTextureCoordinates) const
-{
-	const FBox2D CoordinateBox(InOutTextureCoordinates);
-	const FVector2D CoordinateRange = CoordinateBox.GetSize();
-	const FVector2D Offset = CoordinateBox.Min + Box.Min;
-	const FVector2D Scale = Box.GetSize() / CoordinateRange;
-	for (FVector2D& Coordinate : InOutTextureCoordinates)
-	{
-		Coordinate = (Coordinate - Offset) * Scale;
-	}
-
-}
-
 #undef LOCTEXT_NAMESPACE // "MeshMergeUtils"
diff --git a/Engine/Source/Developer/MeshMergeUtilities/Private/MeshMergeUtilities.h b/Engine/Source/Developer/MeshMergeUtilities/Private/MeshMergeUtilities.h
--- a/Engine/Source/Developer/MeshMergeUtilities/Private/MeshMergeUtilities.h
+++ b/Engine/Source/Developer/MeshMergeUtilities/Private/MeshMergeUtilities.h
@@ -50,8 +50,6 @@ protected:
 	void DetermineMaterialVertexDataUsage(TArray<bool>& InOutMaterialUsesVertexData, const TArray<UMaterialInterface*>& UniqueMaterials, const UMaterialOptions* MaterialOptions) const;
 	/** Creates a proxy material instance at givne path and name */
 	UMaterialInterface* CreateProxyMaterial(const FString &InBasePackageName, FString MergedAssetPackageName, UMaterialInterface* InBaseMaterial, UPackage* InOuter, const FMeshMergingSettings &InSettings, FFlattenMaterial OutMaterial, TArray<UObject *>& OutAssetsToSync, FMaterialUpdateContext* InMaterialUpdateContext = nullptr) const;
-	/** Scales texture coordinates to the specified box, e.g. to 0-1 range in U and V */
-	void ScaleTextureCoordinatesToBox(const FBox2D& Box, TArray<FVector2D>& InOutTextureCoordinates) const;
 	/** Merges flattened material into atlas textures */
 	void MergeFlattenedMaterials(TArray<struct FFlattenMaterial>& InMaterialList, int32 InGutter, FFlattenMaterial& OutMergedMaterial, TArray<FUVOffsetScalePair>& OutUVTransforms) const;
 	/** Merges flattened material into binned textures */
diff --git a/Engine/Source/Runtime/Engine/Classes/Engine/MeshMerging.h b/Engine/Source/Runtime/Engine/Classes/Engine/MeshMerging.h
--- a/Engine/Source/Runtime/Engine/Classes/Engine/MeshMerging.h
+++ b/Engine/Source/Runtime/Engine/Classes/Engine/MeshMerging.h
@@ -307,6 +307,10 @@ struct FMeshProxySettings
 	UPROPERTY(BlueprintReadWrite, EditAnywhere, Category = ProxySettings)
 	uint8 bReuseMeshLightmapUVs:1;
 
+	/** Bake identical meshes (or mesh instances) only once. Can lead to discrepancies with the source mesh visual, especially if WPO is used in the material(s). However, this will result in better quality baked textures & greatly reduce baking time. */
+	UPROPERTY(BlueprintReadWrite, EditAnywhere, Category = ProxySettings)
+	uint8 bGroupIdenticalMeshesForBaking:1;
+
 	/** Whether to generate collision for the merged mesh */
 	UPROPERTY(BlueprintReadWrite, EditAnywhere, Category = ProxySettings)
 	uint8 bCreateCollision:1;
@@ -352,6 +356,7 @@ struct FMeshProxySettings
 		, bUseLandscapeCulling(false)
 		, bAllowDistanceField(false)
 		, bReuseMeshLightmapUVs(true)
+		, bGroupIdenticalMeshesForBaking(false)
 		, bCreateCollision(true)
 		, bAllowVertexColors(false)
 		, bGenerateLightmapUVs(false)
@@ -382,6 +387,7 @@ struct FMeshProxySettings
 			&& bUseLandscapeCulling == Other.bUseLandscapeCulling
 			&& bAllowDistanceField == Other.bAllowDistanceField
 			&& bReuseMeshLightmapUVs == Other.bReuseMeshLightmapUVs
+			&& bGroupIdenticalMeshesForBaking == Other.bGroupIdenticalMeshesForBaking
 			&& bCreateCollision == Other.bCreateCollision
 			&& bAllowVertexColors == Other.bAllowVertexColors
 			&& bGenerateLightmapUVs == Other.bGenerateLightmapUVs
