Fixed issue with PrepareBakingMeshes() asserting when running some code in parallel threads

diff --git a/Engine/Source/Developer/MeshMergeUtilities/Private/MeshMergeUtilities.cpp b/Engine/Source/Developer/MeshMergeUtilities/Private/MeshMergeUtilities.cpp
--- a/Engine/Source/Developer/MeshMergeUtilities/Private/MeshMergeUtilities.cpp
+++ b/Engine/Source/Developer/MeshMergeUtilities/Private/MeshMergeUtilities.cpp
@@ -1315,6 +1315,24 @@ TArray<FMeshData> PrepareBakingMeshes(const struct FMeshProxySettings& InMeshPro
 	TArray<FMeshData> MeshData;
 	MeshData.SetNum(InDescriptors.Num());
 
+	// Step that must run on the game thread
+	for (int32 MeshIndex = 0; MeshIndex < InDescriptors.Num(); MeshIndex++)
+	{
+		FMeshData& MeshSettings = MeshData[MeshIndex];
+		FMeshDescription& MeshDescription = *InMeshDescriptions[MeshIndex];
+		const FProxyMeshDescriptor& MeshDescriptor = InDescriptors[MeshIndex];
+
+		if (!InMeshProxySettings.bGroupIdenticalMeshesForBaking)
+		{
+			if (MeshDescriptor.GetLightMapIndex() != INDEX_NONE)
+			{
+				MeshSettings.LightMap = MeshDescriptor.GetLightMap();
+				MeshSettings.LightMapIndex = MeshDescriptor.GetLightMapIndex();
+			}
+		}
+	}
+
+	// Parallel step
 	ParallelFor(InDescriptors.Num(), [&MeshData, &InDescriptors, &InMeshDescriptions, &InMeshProxySettings](uint32 MeshIndex)
 	{
 		const FProxyMeshDescriptor& MeshDescriptor = InDescriptors[MeshIndex];
@@ -1334,12 +1352,6 @@ TArray<FMeshData> PrepareBakingMeshes(const struct FMeshProxySettings& InMeshPro
 			FMeshDescription& MeshDescription = *InMeshDescriptions[MeshIndex];
 			MeshSettings.MeshDescription = &MeshDescription;
 
-			if (MeshDescriptor.GetLightMapIndex() != INDEX_NONE)
-			{
-				MeshSettings.LightMap = MeshDescriptor.GetLightMap();
-				MeshSettings.LightMapIndex = MeshDescriptor.GetLightMapIndex();
-			}
-
 			TVertexInstanceAttributesRef<FVector2D> VertexInstanceUVs = FStaticMeshAttributes(MeshDescription).GetVertexInstanceUVs();
 
 		    // If we already have lightmap uvs generated and they are valid, we can reuse those instead of having to generate new ones
