Allow the creation of virtual textures when baking HLODs
* Specifying a flatten material that use VT samplers will automatically create textures of the correct type
* Moved default flatten material to config
* Default is now to use VT for HLOD

diff --git a/Engine/Config/BaseEngine.ini b/Engine/Config/BaseEngine.ini
--- a/Engine/Config/BaseEngine.ini
+++ b/Engine/Config/BaseEngine.ini
@@ -136,6 +136,8 @@ EditorBrushMaterialName=/Engine/EngineMaterials/EditorBrushMaterial.EditorBrushM
 DefaultPhysMaterialName=/Engine/EngineMaterials/DefaultPhysicalMaterial.DefaultPhysicalMaterial
 DefaultDeferredDecalMaterialName=/Engine/EngineMaterials/DefaultDeferredDecalMaterial.DefaultDeferredDecalMaterial
 DefaultPostProcessMaterialName=/Engine/EngineMaterials/DefaultPostProcessMaterial.DefaultPostProcessMaterial
+DefaultFlattenMaterialName=/Engine/EngineMaterials/BaseFlattenMaterial.BaseFlattenMaterial
+DefaultHLODFlattenMaterialName=/Engine/EngineMaterials/BaseFlattenMaterial_VT.BaseFlattenMaterial_VT
 TimecodeProviderClassName=None
 ArrowMaterialName=/Engine/EngineMaterials/GizmoMaterial.GizmoMaterial
 ParticleEventManagerClassPath="/Script/Engine.ParticleEventManager"
diff --git a/Engine/Source/Developer/HierarchicalLODUtilities/Private/HierarchicalLODUtilities.cpp b/Engine/Source/Developer/HierarchicalLODUtilities/Private/HierarchicalLODUtilities.cpp
--- a/Engine/Source/Developer/HierarchicalLODUtilities/Private/HierarchicalLODUtilities.cpp
+++ b/Engine/Source/Developer/HierarchicalLODUtilities/Private/HierarchicalLODUtilities.cpp
@@ -44,6 +44,7 @@
 #include "Engine/HLODProxy.h"
 #include "HierarchicalLOD.h"
 #include "LevelUtils.h"
+#include "Materials/Material.h"
 
 DEFINE_LOG_CATEGORY_STATIC(LogHierarchicalLODUtilities, Verbose, All);
 
@@ -633,10 +634,8 @@ bool FHierarchicalLODUtilities::BuildStaticMeshForLODActor(ALODActor* LODActor,
 
 bool FHierarchicalLODUtilities::BuildStaticMeshForLODActor(ALODActor* LODActor, UPackage* AssetsOuter, const FHierarchicalSimplification& LODSetup)
 {
-	UMaterialInterface* BaseMaterial = LoadObject<UMaterialInterface>(NULL, TEXT("/Engine/EngineMaterials/BaseFlattenMaterial.BaseFlattenMaterial"), NULL, LOAD_None, NULL);
-	check(BaseMaterial);
 PRAGMA_DISABLE_DEPRECATION_WARNINGS
-	return BuildStaticMeshForLODActor(LODActor, AssetsOuter, LODSetup, BaseMaterial);
+	return BuildStaticMeshForLODActor(LODActor, AssetsOuter, LODSetup, GEngine->DefaultHLODFlattenMaterial);
 PRAGMA_ENABLE_DEPRECATION_WARNINGS
 }
 
diff --git a/Engine/Source/Developer/MeshMergeUtilities/Private/MeshMergeUtilities.cpp b/Engine/Source/Developer/MeshMergeUtilities/Private/MeshMergeUtilities.cpp
--- a/Engine/Source/Developer/MeshMergeUtilities/Private/MeshMergeUtilities.cpp
+++ b/Engine/Source/Developer/MeshMergeUtilities/Private/MeshMergeUtilities.cpp
@@ -1451,16 +1451,12 @@ TArray<FMeshData> PrepareBakingMeshes(const struct FMeshProxySettings& InMeshPro
 
 void FMeshMergeUtilities::CreateProxyMesh(const TArray<AActor*>& InActors, const struct FMeshProxySettings& InMeshProxySettings, UPackage* InOuter, const FString& InProxyBasePackageName, const FGuid InGuid, const FCreateProxyDelegate& InProxyCreatedDelegate, const bool bAllowAsync, const float ScreenSize) const
 {
-	UMaterial* BaseMaterial = LoadObject<UMaterial>(NULL, TEXT("/Engine/EngineMaterials/BaseFlattenMaterial.BaseFlattenMaterial"), NULL, LOAD_None, NULL);
-	check(BaseMaterial);
-	CreateProxyMesh(InActors, InMeshProxySettings, BaseMaterial, InOuter, InProxyBasePackageName, InGuid, InProxyCreatedDelegate, bAllowAsync, ScreenSize);
+	CreateProxyMesh(InActors, InMeshProxySettings, GEngine->DefaultFlattenMaterial, InOuter, InProxyBasePackageName, InGuid, InProxyCreatedDelegate, bAllowAsync, ScreenSize);
 }
 
 void FMeshMergeUtilities::CreateProxyMesh(const TArray<UStaticMeshComponent*>& InStaticMeshComps, const struct FMeshProxySettings& InMeshProxySettings, UPackage* InOuter, const FString& InProxyBasePackageName, const FGuid InGuid, const FCreateProxyDelegate& InProxyCreatedDelegate, const bool bAllowAsync, const float ScreenSize) const
 {
-	UMaterial* BaseMaterial = LoadObject<UMaterial>(NULL, TEXT("/Engine/EngineMaterials/BaseFlattenMaterial.BaseFlattenMaterial"), NULL, LOAD_None, NULL);
-	check(BaseMaterial);
-	CreateProxyMesh(InStaticMeshComps, InMeshProxySettings, BaseMaterial, InOuter, InProxyBasePackageName, InGuid, InProxyCreatedDelegate, bAllowAsync, ScreenSize);
+	CreateProxyMesh(InStaticMeshComps, InMeshProxySettings, GEngine->DefaultFlattenMaterial, InOuter, InProxyBasePackageName, InGuid, InProxyCreatedDelegate, bAllowAsync, ScreenSize);
 }
 
 void FMeshMergeUtilities::CreateProxyMesh(const TArray<AActor*>& InActors, const struct FMeshProxySettings& InMeshProxySettings, UMaterialInterface* InBaseMaterial, UPackage* InOuter, const FString& InProxyBasePackageName, const FGuid InGuid, const FCreateProxyDelegate& InProxyCreatedDelegate, const bool bAllowAsync /*= false*/, const float ScreenSize /*= 1.0f*/) const
diff --git a/Engine/Source/Developer/MeshMergeUtilities/Private/ProxyMaterialUtilities.h b/Engine/Source/Developer/MeshMergeUtilities/Private/ProxyMaterialUtilities.h
--- a/Engine/Source/Developer/MeshMergeUtilities/Private/ProxyMaterialUtilities.h
+++ b/Engine/Source/Developer/MeshMergeUtilities/Private/ProxyMaterialUtilities.h
@@ -19,13 +19,20 @@ namespace ProxyMaterialUtilities
 	bool b##a##Texture = FlattenMaterial.DoesPropertyContainData(EFlattenMaterialProperties::a) && !FlattenMaterial.IsPropertyConstant(EFlattenMaterialProperties::a); \
 	if ( b##a##Texture) \
 	{ \
+		FName TextureName(#a "Texture"); \
 		UTexture2D* a##Texture = b##a##Texture ? FMaterialUtilities::CreateTexture(InOuter, AssetBasePath + TEXT("T_") + AssetBaseName + TEXT("_" #a), FlattenMaterial.GetPropertySize(EFlattenMaterialProperties::a), FlattenMaterial.GetPropertySamples(EFlattenMaterialProperties::a), b, TEXTUREGROUP_HierarchicalLOD, RF_Public | RF_Standalone, c) : nullptr; \
+		UTexture* DefaultTexture = nullptr; \
+		OutMaterial->GetTextureParameterValue(TextureName, DefaultTexture); \
 		OutMaterial->SetTextureParameterValueEditorOnly(FMaterialParameterInfo(#a "Texture"), a##Texture); \
 		FStaticSwitchParameter SwitchParameter; \
 		SwitchParameter.ParameterInfo.Name = "Use" #a; \
 		SwitchParameter.Value = true; \
 		SwitchParameter.bOverride = true; \
 		NewStaticParameterSet.StaticSwitchParameters.Add(SwitchParameter); \
+		if (ensure(DefaultTexture)) \
+		{ \
+			a##Texture->VirtualTextureStreaming = DefaultTexture->VirtualTextureStreaming; \
+		} \
 		a##Texture->PostEditChange(); \
 		OutAssetsToSync.Add(a##Texture); \
 	} 
@@ -109,7 +116,7 @@ namespace ProxyMaterialUtilities
 		const IMeshMergeUtilities& Module = FModuleManager::Get().LoadModuleChecked<IMeshMergeModule>("MeshMergeUtilities").GetUtilities();
 		if (!Module.IsValidBaseMaterial(InBaseMaterial, false))
 		{
-			BaseMaterial = LoadObject<UMaterialInterface>(NULL, TEXT("/Engine/EngineMaterials/BaseFlattenMaterial.BaseFlattenMaterial"), NULL, LOAD_None, NULL);
+			BaseMaterial = GEngine->DefaultFlattenMaterial;
 		} 
 
 		UMaterialInstanceConstant* OutMaterial = FMaterialUtilities::CreateInstancedMaterial(BaseMaterial, InOuter, AssetBasePath + AssetBaseName, RF_Public | RF_Standalone);
@@ -249,9 +256,20 @@ namespace ProxyMaterialUtilities
 			SwitchParameter.bOverride = true;
 			NewStaticParameterSet.StaticSwitchParameters.Add(SwitchParameter);
 
+
+			const FName PackedTextureName(TEXT("PackedTexture"));
+
+			// Make sure the packed texture is a VT if required by the material.
+			UTexture* DefaultTexture = nullptr;
+			OutMaterial->GetTextureParameterValue(PackedTextureName, DefaultTexture);
+			if (ensure(DefaultTexture))
+			{
+				PackedTexture->VirtualTextureStreaming = DefaultTexture->VirtualTextureStreaming;
+				PackedTexture->PostEditChange();
+			}
+
 			// Set up switch and texture values
-			FMaterialParameterInfo ParameterInfo(TEXT("PackedTexture"));
-			OutMaterial->SetTextureParameterValueEditorOnly(ParameterInfo, PackedTexture);
+			OutMaterial->SetTextureParameterValueEditorOnly(PackedTextureName, PackedTexture);
 		}
 
 		// Emissive is a special case due to the scaling variable
@@ -277,8 +295,7 @@ namespace ProxyMaterialUtilities
 
 	static UMaterialInstanceConstant* CreateProxyMaterialInstance(UPackage* InOuter, const FMaterialProxySettings& InMaterialProxySettings, FFlattenMaterial& FlattenMaterial, const FString& AssetBasePath, const FString& AssetBaseName, TArray<UObject*>& OutAssetsToSync, FMaterialUpdateContext* MaterialUpdateContext = nullptr)
 	{
-		UMaterialInterface* BaseMaterial = LoadObject<UMaterialInterface>(NULL, TEXT("/Engine/EngineMaterials/BaseFlattenMaterial.BaseFlattenMaterial"), NULL, LOAD_None, NULL);
-		return CreateProxyMaterialInstance(InOuter, InMaterialProxySettings, BaseMaterial, FlattenMaterial, AssetBasePath, AssetBasePath, OutAssetsToSync, MaterialUpdateContext);
+		return CreateProxyMaterialInstance(InOuter, InMaterialProxySettings, GEngine->DefaultFlattenMaterial, FlattenMaterial, AssetBasePath, AssetBasePath, OutAssetsToSync, MaterialUpdateContext);
 	}
 
 };
diff --git a/Engine/Source/Editor/StaticMeshEditor/Private/StaticMeshEditorSubsystem.cpp b/Engine/Source/Editor/StaticMeshEditor/Private/StaticMeshEditorSubsystem.cpp
--- a/Engine/Source/Editor/StaticMeshEditor/Private/StaticMeshEditorSubsystem.cpp
+++ b/Engine/Source/Editor/StaticMeshEditor/Private/StaticMeshEditorSubsystem.cpp
@@ -2181,7 +2181,7 @@ bool UStaticMeshEditorSubsystem::CreateProxyMeshActor(const TArray<class AStatic
 	MeshUtilities.CreateProxyMesh(
 		AllActors,                      // List of Actors to merge
 		MergeOptions.MeshProxySettings, // Merge settings
-		nullptr,                        // Base Material used for final proxy material. Note: nullptr for default impl: /Engine/EngineMaterials/BaseFlattenMaterial.BaseFlattenMaterial
+		nullptr,                        // Base Material used for final proxy material. Note: nullptr for default impl: GEngine->DefaultFlattenMaterial
 		nullptr,                        // Package for generated assets. Note: if nullptr, BasePackageName is used
 		PackageName,                    // Will be used for naming generated assets, in case InOuter is not specified ProxyBasePackageName will be used as long package name for creating new packages
 		FGuid::NewGuid(),               // Identify a job, First argument of the ProxyDelegate
diff --git a/Engine/Source/Editor/UnrealEd/Private/HierarchicalLOD.cpp b/Engine/Source/Editor/UnrealEd/Private/HierarchicalLOD.cpp
--- a/Engine/Source/Editor/UnrealEd/Private/HierarchicalLOD.cpp
+++ b/Engine/Source/Editor/UnrealEd/Private/HierarchicalLOD.cpp
@@ -54,7 +54,10 @@ DEFINE_LOG_CATEGORY_STATIC(LogLODGenerator, Log, All);
 UHierarchicalLODSettings::UHierarchicalLODSettings(const FObjectInitializer& ObjectInitializer)
 	: Super(ObjectInitializer), bForceSettingsInAllMaps(false), BaseMaterial(nullptr)
 {	
-	BaseMaterial = ConstructorHelpers::FObjectFinder<UMaterialInterface>(TEXT("/Engine/EngineMaterials/BaseFlattenMaterial")).Object;
+	if (!IsTemplate())
+	{
+		BaseMaterial = GEngine->DefaultHLODFlattenMaterial;
+	}
 }
 
 void UHierarchicalLODSettings::PostEditChangeProperty(struct FPropertyChangedEvent& PropertyChangedEvent)
@@ -66,7 +69,7 @@ void UHierarchicalLODSettings::PostEditChangeProperty(struct FPropertyChangedEve
 			const IMeshMergeUtilities& Module = FModuleManager::Get().LoadModuleChecked<IMeshMergeModule>("MeshMergeUtilities").GetUtilities();
 			if (!Module.IsValidBaseMaterial(BaseMaterial.LoadSynchronous(), true))
 			{
-				BaseMaterial = LoadObject<UMaterialInterface>(NULL, TEXT("/Engine/EngineMaterials/BaseFlattenMaterial.BaseFlattenMaterial"), NULL, LOAD_None, NULL);
+				BaseMaterial = GEngine->DefaultFlattenMaterial;
 			}
 		}
 	}
diff --git a/Engine/Source/Runtime/Engine/Classes/Engine/Engine.h b/Engine/Source/Runtime/Engine/Classes/Engine/Engine.h
--- a/Engine/Source/Runtime/Engine/Classes/Engine/Engine.h
+++ b/Engine/Source/Runtime/Engine/Classes/Engine/Engine.h
@@ -1009,6 +1009,20 @@ public:
 	/** A material used to render debug meshes. */
 	UPROPERTY()
 	TObjectPtr<class UMaterial> DebugEditorMaterial;
+
+	/** A material used to flatten materials. */
+	UPROPERTY(globalconfig)
+	FSoftObjectPath DefaultFlattenMaterialName;
+
+	/** A material used to flatten materials to VT textures. */
+	UPROPERTY(globalconfig)
+	FSoftObjectPath DefaultHLODFlattenMaterialName;
+
+	/** Materials used when flattening materials */
+	UPROPERTY()
+	TObjectPtr<class UMaterial> DefaultFlattenMaterial;
+	UPROPERTY()
+	TObjectPtr<class UMaterial> DefaultHLODFlattenMaterial;
 #endif
 
 	/** A material used to render debug opaque material. Used in various animation editor viewport features. */
diff --git a/Engine/Source/Runtime/Engine/Private/UnrealEngine.cpp b/Engine/Source/Runtime/Engine/Private/UnrealEngine.cpp
--- a/Engine/Source/Runtime/Engine/Private/UnrealEngine.cpp
+++ b/Engine/Source/Runtime/Engine/Private/UnrealEngine.cpp
@@ -2703,6 +2703,11 @@ void UEngine::InitializeObjectReferences()
 		}
 	}
 
+#if WITH_EDITORONLY_DATA
+	LoadSpecialMaterial(TEXT("DefaultFlattenMaterialName"), DefaultFlattenMaterialName.ToString(), DefaultFlattenMaterial, false);
+	LoadSpecialMaterial(TEXT("DefaultHLODFlattenMaterialName"), DefaultHLODFlattenMaterialName.ToString(), DefaultHLODFlattenMaterial, false);
+#endif
+
 	LoadEngineTexture(DefaultTexture, *DefaultTextureName.ToString());
 	LoadEngineTexture(DefaultDiffuseTexture, *DefaultDiffuseTextureName.ToString());
 	LoadEngineTexture(HighFrequencyNoiseTexture, *HighFrequencyNoiseTextureName.ToString());
diff --git a/Engine/Source/Runtime/Engine/Private/WorldPartition/HLOD/HLODLayer.cpp b/Engine/Source/Runtime/Engine/Private/WorldPartition/HLOD/HLODLayer.cpp
--- a/Engine/Source/Runtime/Engine/Private/WorldPartition/HLOD/HLODLayer.cpp
+++ b/Engine/Source/Runtime/Engine/Private/WorldPartition/HLOD/HLODLayer.cpp
@@ -29,7 +29,10 @@ UHLODLayer::UHLODLayer(const FObjectInitializer& ObjectInitializer)
 #endif
 {
 #if WITH_EDITORONLY_DATA
-	HLODMaterial = ConstructorHelpers::FObjectFinder<UMaterialInterface>(TEXT("/Engine/EngineMaterials/BaseFlattenMaterial")).Object;
+	if (!IsTemplate())
+	{
+		HLODMaterial = GEngine->DefaultHLODFlattenMaterial;
+	}
 #endif
 }
 
diff --git a/Engine/Source/Runtime/Engine/Private/WorldSettings.cpp b/Engine/Source/Runtime/Engine/Private/WorldSettings.cpp
--- a/Engine/Source/Runtime/Engine/Private/WorldSettings.cpp
+++ b/Engine/Source/Runtime/Engine/Private/WorldSettings.cpp
@@ -750,7 +750,7 @@ void AWorldSettings::InternalPostPropertyChanged(FName PropertyName)
 				const IMeshMergeUtilities& Module = FModuleManager::Get().LoadModuleChecked<IMeshMergeModule>("MeshMergeUtilities").GetUtilities();
 				if (!Module.IsValidBaseMaterial(OverrideBaseMaterial.LoadSynchronous(), true))
 				{
-					OverrideBaseMaterial = LoadObject<UMaterialInterface>(NULL, TEXT("/Engine/EngineMaterials/BaseFlattenMaterial.BaseFlattenMaterial"), NULL, LOAD_None, NULL);
+					OverrideBaseMaterial = GEngine->DefaultHLODFlattenMaterial;
 				}
 			}
 		}
@@ -766,7 +766,7 @@ void UHierarchicalLODSetup::PostEditChangeProperty(struct FPropertyChangedEvent&
 			const IMeshMergeUtilities& Module = FModuleManager::Get().LoadModuleChecked<IMeshMergeModule>("MeshMergeUtilities").GetUtilities();
 			if (!Module.IsValidBaseMaterial(OverrideBaseMaterial.LoadSynchronous(), true))
 			{
-				OverrideBaseMaterial = LoadObject<UMaterialInterface>(NULL, TEXT("/Engine/EngineMaterials/BaseFlattenMaterial.BaseFlattenMaterial"), NULL, LOAD_None, NULL);
+				OverrideBaseMaterial = GEngine->DefaultHLODFlattenMaterial;
 			}
 		}
 	}
diff --git a/Engine/Source/Runtime/Engine/Public/WorldPartition/HLOD/HLODLayer.h b/Engine/Source/Runtime/Engine/Public/WorldPartition/HLOD/HLODLayer.h
--- a/Engine/Source/Runtime/Engine/Public/WorldPartition/HLOD/HLODLayer.h
+++ b/Engine/Source/Runtime/Engine/Public/WorldPartition/HLOD/HLODLayer.h
@@ -65,7 +65,7 @@ private:
 	FMeshProxySettings MeshSimplifySettings;
 
 	/** Material that will be used by the generated HLOD static mesh */
-	UPROPERTY(EditAnywhere, Config, Category=HLOD, meta = (editcondition = "LayerType == EHLODLayerType::MeshMerge || LayerType == EHLODLayerType::MeshSimplify"))
+	UPROPERTY(EditAnywhere, Config, AdvancedDisplay, Category=HLOD, meta = (editcondition = "LayerType == EHLODLayerType::MeshMerge || LayerType == EHLODLayerType::MeshSimplify"))
 	TSoftObjectPtr<UMaterial> HLODMaterial;
 
 	UPROPERTY(EditAnywhere, Config, Category=HLOD)
