Avoid renderdata of HLOD being recomputed during cook package save because of non-deterministic guid
  - Preserve guid during save if static mesh bulkdata content hasn't changed to avoid triggering unnecessary rebuilds
  - Make sure the next HLOD generation is going to use the HashAsGuid feature to avoid recompiling staticmesh that are generated if they are exactly the same as before

  - 31m10s to 14m15s to cook P_Construct_WP with a warmed-up DDC

diff --git a/Engine/Source/Developer/MeshMergeUtilities/Private/MeshMergeUtilities.cpp b/Engine/Source/Developer/MeshMergeUtilities/Private/MeshMergeUtilities.cpp
--- a/Engine/Source/Developer/MeshMergeUtilities/Private/MeshMergeUtilities.cpp
+++ b/Engine/Source/Developer/MeshMergeUtilities/Private/MeshMergeUtilities.cpp
@@ -3020,7 +3020,10 @@ void FMeshMergeUtilities::MergeComponentsToStaticMesh(const TArray<UPrimitiveCom
 				}
 
 				FMeshDescription* MeshDescription = StaticMesh->CreateMeshDescription(LODIndex, MergedMeshLOD);
-				StaticMesh->CommitMeshDescription(LODIndex);
+
+				UStaticMesh::FCommitMeshDescriptionParams CommitParams;
+				CommitParams.bUseHashAsGuid = true;
+				StaticMesh->CommitMeshDescription(LODIndex, CommitParams);
 			}
 		}
 		
