Fixed test failure - ActorMerging_SingleLODMaterialBaking

Corrected bad conversion from EMaterialProperty to EFlattenMaterialProperties using array indexing... now using a map.
Fixed slight delta in ground truth images as the floor pattern moved a bit (irrelevant to the actor merging test)

diff --git a/Engine/Source/Developer/MeshMergeUtilities/Private/MeshMergeUtilities.cpp b/Engine/Source/Developer/MeshMergeUtilities/Private/MeshMergeUtilities.cpp
--- a/Engine/Source/Developer/MeshMergeUtilities/Private/MeshMergeUtilities.cpp
+++ b/Engine/Source/Developer/MeshMergeUtilities/Private/MeshMergeUtilities.cpp
@@ -537,7 +537,7 @@ void FMeshMergeUtilities::ConvertOutputToFlatMaterials(const TArray<FBakeOutput>
 
 		for (TPair<EMaterialProperty, FIntPoint> SizePair : Output.PropertySizes)
 		{
-			EFlattenMaterialProperties OldProperty = NewToOldProperty(SizePair.Key);
+			EFlattenMaterialProperties OldProperty = ToFlattenProperty(SizePair.Key);
 			Material.SetPropertySize(OldProperty, SizePair.Value);
 			Material.GetPropertySamples(OldProperty).Append(Output.PropertyData[SizePair.Key]);
 		}
@@ -566,7 +566,7 @@ void FMeshMergeUtilities::TransferOutputToFlatMaterials(const TArray<FMaterialDa
 
 		for (TPair<EMaterialProperty, FIntPoint> SizePair : Output.PropertySizes)
 		{
-			EFlattenMaterialProperties OldProperty = NewToOldProperty(SizePair.Key);
+			EFlattenMaterialProperties OldProperty = ToFlattenProperty(SizePair.Key);
 			Material.SetPropertySize(OldProperty, SizePair.Value);
 			Material.GetPropertySamples(OldProperty) = MoveTemp(Output.PropertyData[SizePair.Key]);
 		}
@@ -578,30 +578,23 @@ void FMeshMergeUtilities::TransferOutputToFlatMaterials(const TArray<FMaterialDa
 	}
 }
 
-EFlattenMaterialProperties FMeshMergeUtilities::NewToOldProperty(int32 NewProperty) const
+EFlattenMaterialProperties FMeshMergeUtilities::ToFlattenProperty(EMaterialProperty MaterialProperty) const
 {
-	const EFlattenMaterialProperties Remap[MP_Refraction] =
-	{
-		EFlattenMaterialProperties::Emissive,
-		EFlattenMaterialProperties::Opacity,
-		EFlattenMaterialProperties::OpacityMask,
-		EFlattenMaterialProperties::NumFlattenMaterialProperties,
-		EFlattenMaterialProperties::NumFlattenMaterialProperties,
-		EFlattenMaterialProperties::Diffuse,
-		EFlattenMaterialProperties::Metallic,
-		EFlattenMaterialProperties::Specular,
-		EFlattenMaterialProperties::Roughness,
-		EFlattenMaterialProperties::Anisotropy,
-		EFlattenMaterialProperties::Normal,
-		EFlattenMaterialProperties::Tangent,
-		EFlattenMaterialProperties::NumFlattenMaterialProperties,
-		EFlattenMaterialProperties::NumFlattenMaterialProperties,
-		EFlattenMaterialProperties::NumFlattenMaterialProperties,
-		EFlattenMaterialProperties::NumFlattenMaterialProperties,
-		EFlattenMaterialProperties::AmbientOcclusion
-	};
-	
-	return Remap[NewProperty];
+	switch (MaterialProperty)
+	{
+	case EMaterialProperty::MP_BaseColor:			return EFlattenMaterialProperties::Diffuse;
+	case EMaterialProperty::MP_Metallic:			return EFlattenMaterialProperties::Metallic;
+	case EMaterialProperty::MP_Specular:			return EFlattenMaterialProperties::Specular;
+	case EMaterialProperty::MP_Roughness:			return EFlattenMaterialProperties::Roughness;
+	case EMaterialProperty::MP_Anisotropy:			return EFlattenMaterialProperties::Anisotropy;
+	case EMaterialProperty::MP_Normal:				return EFlattenMaterialProperties::Normal;
+	case EMaterialProperty::MP_Tangent:				return EFlattenMaterialProperties::Tangent;
+	case EMaterialProperty::MP_Opacity:				return EFlattenMaterialProperties::Opacity;
+	case EMaterialProperty::MP_EmissiveColor:		return EFlattenMaterialProperties::Emissive;
+	case EMaterialProperty::MP_OpacityMask:			return EFlattenMaterialProperties::OpacityMask;
+	case EMaterialProperty::MP_AmbientOcclusion:	return EFlattenMaterialProperties::AmbientOcclusion;
+	default:										return EFlattenMaterialProperties::NumFlattenMaterialProperties;
+	}
 }
 
 UMaterialOptions* FMeshMergeUtilities::PopulateMaterialOptions(const FMaterialProxySettings& MaterialSettings) const
@@ -2655,7 +2648,7 @@ void FMeshMergeUtilities::MergeComponentsToStaticMesh(const TArray<UPrimitiveCom
 		{
 			if (Entry.Property != MP_MAX)
 			{
-				EFlattenMaterialProperties OldProperty = NewToOldProperty(Entry.Property);
+				EFlattenMaterialProperties OldProperty = ToFlattenProperty(Entry.Property);
 				OutMaterial.SetPropertySize(OldProperty, Entry.bUseCustomSize ? Entry.CustomSize : MaterialOptions->TextureSize);
 			}
 		}
diff --git a/Engine/Source/Developer/MeshMergeUtilities/Private/MeshMergeUtilities.h b/Engine/Source/Developer/MeshMergeUtilities/Private/MeshMergeUtilities.h
--- a/Engine/Source/Developer/MeshMergeUtilities/Private/MeshMergeUtilities.h
+++ b/Engine/Source/Developer/MeshMergeUtilities/Private/MeshMergeUtilities.h
@@ -76,8 +76,9 @@ protected:
 	  * @warning This is a destructive operation for InOutBakeOutputs
 	  */
 	void TransferOutputToFlatMaterials(const TArray<FMaterialData>& InMaterialData, TArray<FBakeOutput>& InOutBakeOutputs, TArray<FFlattenMaterial> &OutFlattenedMaterials) const;
-	/** Converts new material property value to old legacy enum values */
-	EFlattenMaterialProperties NewToOldProperty(int32 OldProperty) const;
+
+	/** Converts material property enum value to flatten material property */
+	EFlattenMaterialProperties ToFlattenProperty(EMaterialProperty MaterialProperty) const;
 private:
 	FProxyGenerationProcessor* Processor;
 	FDelegateHandle ModuleLoadedDelegateHandle;
