MaterialBaking - Use correct linear bake for all texture property by default
Providing a new flatten material with the correct texture samplers.
Retrieve texture properties from the default textures in the material (linear vs srgb, compression, vt enabled, lod group...)

Previously generated material instances are unchanged.
All MergeActor engine tests pass successfully.

diff --git a/Engine/Config/BaseEngine.ini b/Engine/Config/BaseEngine.ini
--- a/Engine/Config/BaseEngine.ini
+++ b/Engine/Config/BaseEngine.ini
@@ -136,8 +136,8 @@ EditorBrushMaterialName=/Engine/EngineMaterials/EditorBrushMaterial.EditorBrushM
 DefaultPhysMaterialName=/Engine/EngineMaterials/DefaultPhysicalMaterial.DefaultPhysicalMaterial
 DefaultDeferredDecalMaterialName=/Engine/EngineMaterials/DefaultDeferredDecalMaterial.DefaultDeferredDecalMaterial
 DefaultPostProcessMaterialName=/Engine/EngineMaterials/DefaultPostProcessMaterial.DefaultPostProcessMaterial
-DefaultFlattenMaterialName=/Engine/EngineMaterials/BaseFlattenMaterial.BaseFlattenMaterial
-DefaultHLODFlattenMaterialName=/Engine/EngineMaterials/BaseFlattenMaterial_VT.BaseFlattenMaterial_VT
+DefaultFlattenMaterialName=/Engine/EngineMaterials/FlattenMaterial.FlattenMaterial
+DefaultHLODFlattenMaterialName=/Engine/EngineMaterials/FlattenMaterial_VT.FlattenMaterial_VT
 TimecodeProviderClassName=None
 ArrowMaterialName=/Engine/EngineMaterials/GizmoMaterial.GizmoMaterial
 ParticleEventManagerClassPath="/Script/Engine.ParticleEventManager"
diff --git a/Engine/Source/Developer/MaterialBaking/Private/MaterialBakingModule.cpp b/Engine/Source/Developer/MaterialBaking/Private/MaterialBakingModule.cpp
--- a/Engine/Source/Developer/MaterialBaking/Private/MaterialBakingModule.cpp
+++ b/Engine/Source/Developer/MaterialBaking/Private/MaterialBakingModule.cpp
@@ -261,7 +261,7 @@ void FMaterialBakingModule::StartupModule()
 	bEmissiveHDR = false;
 
 	// Set which properties should enforce gamma correction
-	SetLinearBake(false);
+	SetLinearBake(true);
 
 	// Set which pixel format should be used for the possible baked out material properties
 	PerPropertyFormat.Add(MP_EmissiveColor, PF_FloatRGBA);
diff --git a/Engine/Source/Developer/MaterialUtilities/Private/MaterialUtilities.cpp b/Engine/Source/Developer/MaterialUtilities/Private/MaterialUtilities.cpp
--- a/Engine/Source/Developer/MaterialUtilities/Private/MaterialUtilities.cpp
+++ b/Engine/Source/Developer/MaterialUtilities/Private/MaterialUtilities.cpp
@@ -75,7 +75,7 @@ void FMaterialUtilities::OnPreGarbageCollect()
 
 UMaterialInterface* FMaterialUtilities::CreateProxyMaterialAndTextures(UPackage* OuterPackage, const FString& AssetName, const FBakeOutput& BakeOutput, const FMeshData& MeshData, const FMaterialData& MaterialData, UMaterialOptions* Options)
 {
-	TArray<EMaterialProperty> SRGBDisabledProperties{ MP_Opacity, MP_OpacityMask, MP_Normal, MP_EmissiveColor };
+	TArray<EMaterialProperty> SRGBEnabledProperties{ MP_BaseColor, MP_EmissiveColor, MP_SubsurfaceColor };
 
 	// Certain material properties use differen compression settings
 	TMap<EMaterialProperty, TextureCompressionSettings> SpecialCompressionSettingProperties;
@@ -84,7 +84,7 @@ UMaterialInterface* FMaterialUtilities::CreateProxyMaterialAndTextures(UPackage*
 	SpecialCompressionSettingProperties.Add(MP_OpacityMask, TC_Grayscale);
 	SpecialCompressionSettingProperties.Add(MP_AmbientOcclusion, TC_Grayscale);
 
-	UMaterial* BaseMaterial = LoadObject<UMaterial>(nullptr, TEXT("/Engine/EngineMaterials/NewBaseFlattenMaterial.NewBaseFlattenMaterial"), nullptr, LOAD_None, nullptr);
+	UMaterial* BaseMaterial = GEngine->DefaultFlattenMaterial;
 	check(BaseMaterial);
 
 	/** Create Proxy material and populate flags */
@@ -119,7 +119,7 @@ UMaterialInterface* FMaterialUtilities::CreateProxyMaterialAndTextures(UPackage*
 		if (ColorData.Num() > 1)
 		{
 			TextureCompressionSettings CompressionSettings = SpecialCompressionSettingProperties.Contains(Property) ? SpecialCompressionSettingProperties.FindChecked(Property) : TC_Default;
-			bool bSRGBEnabled = !SRGBDisabledProperties.Contains(Property);
+			bool bSRGBEnabled = SRGBEnabledProperties.Contains(Property);
 			UTexture* Texture = FMaterialUtilities::CreateTexture(OuterPackage, TEXT("T_") + AssetName + TEXT("_") + TrimmedPropertyName, DataSize, ColorData, CompressionSettings, TEXTUREGROUP_HierarchicalLOD, RF_Public | RF_Standalone, bSRGBEnabled);
 
 			// Set texture parameter value on instance material
diff --git a/Engine/Source/Developer/MeshMergeUtilities/Private/ProxyMaterialUtilities.h b/Engine/Source/Developer/MeshMergeUtilities/Private/ProxyMaterialUtilities.h
--- a/Engine/Source/Developer/MeshMergeUtilities/Private/ProxyMaterialUtilities.h
+++ b/Engine/Source/Developer/MeshMergeUtilities/Private/ProxyMaterialUtilities.h
@@ -15,47 +15,6 @@
 
 namespace ProxyMaterialUtilities
 {
-#define TEXTURE_MACRO_BASE(a, b, c) \
-	bool b##a##Texture = FlattenMaterial.DoesPropertyContainData(EFlattenMaterialProperties::a) && !FlattenMaterial.IsPropertyConstant(EFlattenMaterialProperties::a); \
-	if ( b##a##Texture) \
-	{ \
-		FName TextureName(#a "Texture"); \
-		UTexture2D* a##Texture = b##a##Texture ? FMaterialUtilities::CreateTexture(InOuter, AssetBasePath + TEXT("T_") + AssetBaseName + TEXT("_" #a), FlattenMaterial.GetPropertySize(EFlattenMaterialProperties::a), FlattenMaterial.GetPropertySamples(EFlattenMaterialProperties::a), b, TEXTUREGROUP_HierarchicalLOD, RF_Public | RF_Standalone, c) : nullptr; \
-		UTexture* DefaultTexture = nullptr; \
-		OutMaterial->GetTextureParameterValue(TextureName, DefaultTexture); \
-		OutMaterial->SetTextureParameterValueEditorOnly(FMaterialParameterInfo(#a "Texture"), a##Texture); \
-		FStaticSwitchParameter SwitchParameter; \
-		SwitchParameter.ParameterInfo.Name = "Use" #a; \
-		SwitchParameter.Value = true; \
-		SwitchParameter.bOverride = true; \
-		NewStaticParameterSet.StaticSwitchParameters.Add(SwitchParameter); \
-		if (ensure(DefaultTexture)) \
-		{ \
-			a##Texture->VirtualTextureStreaming = DefaultTexture->VirtualTextureStreaming; \
-		} \
-		a##Texture->PostEditChange(); \
-		OutAssetsToSync.Add(a##Texture); \
-	} 
-
-#define TEXTURE_MACRO_VECTOR(a, b, c) TEXTURE_MACRO_BASE(a, b, c)\
-	else\
-	{ \
-		OutMaterial->SetVectorParameterValueEditorOnly(FMaterialParameterInfo(#a "Const"), FlattenMaterial.GetPropertySamples(EFlattenMaterialProperties::a)[0]); \
-	} 
-
-#define TEXTURE_MACRO_VECTOR_LINEAR(a, b, c) TEXTURE_MACRO_BASE(a, b, c)\
-	else\
-	{ \
-		OutMaterial->SetVectorParameterValueEditorOnly(FMaterialParameterInfo(#a "Const"), FlattenMaterial.GetPropertySamples(EFlattenMaterialProperties::a)[0].ReinterpretAsLinear()); \
-	} 
-
-#define TEXTURE_MACRO_SCALAR(a, b, c) TEXTURE_MACRO_BASE(a, b, c)\
-	else \
-	{ \
-		FLinearColor Colour = FlattenMaterial.IsPropertyConstant(EFlattenMaterialProperties::a) ? FLinearColor::FromSRGBColor( FlattenMaterial.GetPropertySamples(EFlattenMaterialProperties::a)[0]) : FLinearColor( InMaterialProxySettings.a##Constant, 0, 0, 0 ); \
-		OutMaterial->SetScalarParameterValueEditorOnly(FMaterialParameterInfo(#a "Const"), Colour.R); \
-	}
-
 	static const bool CalculatePackedTextureData(const FFlattenMaterial& InMaterial, bool& bOutPackMetallic, bool& bOutPackSpecular, bool& bOutPackRoughness, int32& OutNumSamples, FIntPoint& OutSize)
 	{
 		// Whether or not a material property is baked down
@@ -138,9 +97,6 @@ namespace ProxyMaterialUtilities
 		FIntPoint PackedSize;
 		const bool bPackTextures = CalculatePackedTextureData(FlattenMaterial, bPackMetallic, bPackSpecular, bPackRoughness, NumSamples, PackedSize);
 
-		const bool bSRGB = true;
-		const bool bRGB = false;
-
 		FStaticParameterSet NewStaticParameterSet;
 
 		if(FlattenMaterial.UVChannel != 0)
@@ -156,48 +112,117 @@ namespace ProxyMaterialUtilities
 			NewStaticParameterSet.StaticSwitchParameters.Add(SwitchParameter);
 		}
 
+		auto CreateTextureFromDefault = [&](const FName TextureName, const FString& AssetLongName, FIntPoint Size, const TArray<FColor>& Samples)
+		{
+			bool bSRGB = false;
+			bool bVirtualTextureStreaming = false;
+			TextureCompressionSettings CompressionSettings = TextureCompressionSettings::TC_Default;
+			TextureGroup LODGroup = TextureGroup::TEXTUREGROUP_World;
+
+			UTexture* DefaultTexture = nullptr;
+			OutMaterial->GetTextureParameterValue(TextureName, DefaultTexture);
+			if (ensure(DefaultTexture))
+			{
+				bSRGB = DefaultTexture->SRGB;
+				bVirtualTextureStreaming = DefaultTexture->VirtualTextureStreaming;
+				CompressionSettings = DefaultTexture->CompressionSettings;
+				LODGroup = DefaultTexture->LODGroup;
+			}
+
+			UTexture2D* Texture = FMaterialUtilities::CreateTexture(InOuter, AssetLongName, Size, Samples, CompressionSettings, LODGroup, RF_Public | RF_Standalone, bSRGB);
+			Texture->VirtualTextureStreaming = bVirtualTextureStreaming;
+			Texture->PostEditChange();
+			return Texture;
+		};
+
+		auto SetTextureParam = [&](const FString& PropertyName, EFlattenMaterialProperties FlattenProperty)
+		{
+			if (FlattenMaterial.DoesPropertyContainData(FlattenProperty) && !FlattenMaterial.IsPropertyConstant(FlattenProperty))
+			{
+				const FName TextureName(PropertyName + TEXT("Texture"));
+				const FName UseTexture(TEXT("Use") + PropertyName);
+
+				UTexture2D* Texture = CreateTextureFromDefault(TextureName, AssetBasePath + TEXT("T_") + AssetBaseName + TEXT("_") + PropertyName, FlattenMaterial.GetPropertySize(FlattenProperty), FlattenMaterial.GetPropertySamples(FlattenProperty));
+
+				FStaticSwitchParameter SwitchParameter;
+				SwitchParameter.ParameterInfo.Name = UseTexture;
+				SwitchParameter.Value = true;
+				SwitchParameter.bOverride = true;
+				NewStaticParameterSet.StaticSwitchParameters.Add(SwitchParameter);
+
+				OutMaterial->SetTextureParameterValueEditorOnly(TextureName, Texture);
+			
+				OutAssetsToSync.Add(Texture);
+			}
+		};
+
+		auto SetTextureParamConstVector = [&](const FString& PropertyName, EFlattenMaterialProperties FlattenProperty)
+		{
+			if (FlattenMaterial.DoesPropertyContainData(FlattenProperty) && !FlattenMaterial.IsPropertyConstant(FlattenProperty))
+			{
+				SetTextureParam(PropertyName, FlattenProperty);
+			} 
+			else
+			{
+				const FName ConstName(PropertyName + TEXT("Const"));
+				OutMaterial->SetVectorParameterValueEditorOnly(ConstName, FlattenMaterial.GetPropertySamples(FlattenProperty)[0]);
+			} 
+		};
+
+		auto SetTextureParamConstScalar = [&](const FString& PropertyName, EFlattenMaterialProperties FlattenProperty, float ConstantValue)
+		{
+			if (FlattenMaterial.DoesPropertyContainData(FlattenProperty) && !FlattenMaterial.IsPropertyConstant(FlattenProperty))
+			{
+				SetTextureParam(PropertyName, FlattenProperty);
+			}
+			else
+			{
+				const FName ConstName(PropertyName + TEXT("Const"));
+				FLinearColor Colour = FlattenMaterial.IsPropertyConstant(FlattenProperty) ? FLinearColor::FromSRGBColor(FlattenMaterial.GetPropertySamples(FlattenProperty)[0]) : FLinearColor(ConstantValue, 0, 0, 0);
+				OutMaterial->SetScalarParameterValueEditorOnly(ConstName, Colour.R);
+			}
+		};
+
 		// Load textures and set switches accordingly
 		if (FlattenMaterial.GetPropertySize(EFlattenMaterialProperties::Diffuse).Num() > 0 && !(FlattenMaterial.IsPropertyConstant(EFlattenMaterialProperties::Diffuse) && FlattenMaterial.GetPropertySamples(EFlattenMaterialProperties::Diffuse)[0] == FColor::Black))
 		{
-			TEXTURE_MACRO_VECTOR(Diffuse, TC_Default, bSRGB);
+			SetTextureParamConstVector("BaseColor", EFlattenMaterialProperties::Diffuse);
 		}
 
 		if (FlattenMaterial.GetPropertySize(EFlattenMaterialProperties::Normal).Num() > 1)
 		{
-			TEXTURE_MACRO_BASE(Normal, TC_Normalmap, bRGB)
+			SetTextureParam("Normal", EFlattenMaterialProperties::Normal);
 		}
 
 		// Determine whether or not specific material properties are packed together into one texture (requires at least two to match (number of samples and texture size) in order to be packed
 		if (!bPackMetallic && (FlattenMaterial.GetPropertySize(EFlattenMaterialProperties::Metallic).Num() > 0 || !InMaterialProxySettings.bMetallicMap))
 		{
-			TEXTURE_MACRO_SCALAR(Metallic, TC_Default, bSRGB);
+			SetTextureParamConstScalar("Metallic", EFlattenMaterialProperties::Metallic, InMaterialProxySettings.MetallicConstant);
 		}
 
 		if (!bPackRoughness && (FlattenMaterial.GetPropertySize(EFlattenMaterialProperties::Roughness).Num() > 0 || !InMaterialProxySettings.bRoughnessMap))
 		{
-			TEXTURE_MACRO_SCALAR(Roughness, TC_Default, bSRGB);
+			SetTextureParamConstScalar("Roughness", EFlattenMaterialProperties::Roughness, InMaterialProxySettings.RoughnessConstant);
 		}
 
 		if (!bPackSpecular && (FlattenMaterial.GetPropertySize(EFlattenMaterialProperties::Specular).Num() > 0 || !InMaterialProxySettings.bSpecularMap))
 		{
-			TEXTURE_MACRO_SCALAR(Specular, TC_Default, bSRGB);
+			SetTextureParamConstScalar("Specular", EFlattenMaterialProperties::Specular, InMaterialProxySettings.SpecularConstant);
 		}
 
-		const bool bNonSRGB = false;
-
 		if (FlattenMaterial.GetPropertySize(EFlattenMaterialProperties::Opacity).Num() > 0 || !InMaterialProxySettings.bOpacityMap)
 		{
-			TEXTURE_MACRO_SCALAR(Opacity, TC_Grayscale, bNonSRGB);
+			SetTextureParamConstScalar("Opacity", EFlattenMaterialProperties::Opacity, InMaterialProxySettings.OpacityConstant);
 		}
 
 		if (FlattenMaterial.GetPropertySize(EFlattenMaterialProperties::OpacityMask).Num() > 0 || !InMaterialProxySettings.bOpacityMaskMap)
 		{
-			TEXTURE_MACRO_SCALAR(OpacityMask, TC_Grayscale, bNonSRGB);
+			SetTextureParamConstScalar("OpacityMask", EFlattenMaterialProperties::OpacityMask, InMaterialProxySettings.OpacityMaskConstant);
 		}
 
 		if (FlattenMaterial.GetPropertySize(EFlattenMaterialProperties::AmbientOcclusion).Num() > 0 || !InMaterialProxySettings.bAmbientOcclusionMap)
 		{
-			TEXTURE_MACRO_SCALAR(AmbientOcclusion, TC_Grayscale, bNonSRGB);
+			SetTextureParamConstScalar("AmbientOcclusion", EFlattenMaterialProperties::AmbientOcclusion, InMaterialProxySettings.AmbientOcclusionConstant);
 		}
 
 		// Handle the packed texture if applicable
@@ -235,7 +260,8 @@ namespace ProxyMaterialUtilities
 			}
 
 			// Create texture using the merged property data
-			UTexture2D* PackedTexture = FMaterialUtilities::CreateTexture(InOuter, AssetBasePath + TEXT("T_") + AssetBaseName + TEXT("_MRS"), PackedSize, MergedTexture, TC_Default, TEXTUREGROUP_HierarchicalLOD, RF_Public | RF_Standalone, bSRGB);
+			const FName PackedTextureName(TEXT("PackedTexture"));
+			UTexture2D* PackedTexture = CreateTextureFromDefault(PackedTextureName, AssetBasePath + TEXT("T_") + AssetBaseName + TEXT("_MRS"), PackedSize, MergedTexture);
 			checkf(PackedTexture, TEXT("Failed to create texture"));
 			OutAssetsToSync.Add(PackedTexture);
 
@@ -256,18 +282,6 @@ namespace ProxyMaterialUtilities
 			SwitchParameter.bOverride = true;
 			NewStaticParameterSet.StaticSwitchParameters.Add(SwitchParameter);
 
-
-			const FName PackedTextureName(TEXT("PackedTexture"));
-
-			// Make sure the packed texture is a VT if required by the material.
-			UTexture* DefaultTexture = nullptr;
-			OutMaterial->GetTextureParameterValue(PackedTextureName, DefaultTexture);
-			if (ensure(DefaultTexture))
-			{
-				PackedTexture->VirtualTextureStreaming = DefaultTexture->VirtualTextureStreaming;
-				PackedTexture->PostEditChange();
-			}
-
 			// Set up switch and texture values
 			OutMaterial->SetTextureParameterValueEditorOnly(PackedTextureName, PackedTexture);
 		}
@@ -275,7 +289,7 @@ namespace ProxyMaterialUtilities
 		// Emissive is a special case due to the scaling variable
 		if (FlattenMaterial.GetPropertySamples(EFlattenMaterialProperties::Emissive).Num() > 0 && !(FlattenMaterial.GetPropertySize(EFlattenMaterialProperties::Emissive).Num() == 1 && FlattenMaterial.GetPropertySamples(EFlattenMaterialProperties::Emissive)[0] == FColor::Black))
 		{
-			TEXTURE_MACRO_VECTOR_LINEAR(Emissive, TC_Default, bRGB);
+			SetTextureParamConstVector("EmissiveColor", EFlattenMaterialProperties::Emissive);
 
 			if (FlattenMaterial.EmissiveScale != 1.0f)
 			{
