FMeshMergeHelpers::ExportStaticMeshLOD() - Added missing Register() call
* Fixes a crash accessing some of the attributes later on

diff --git a/Engine/Source/Developer/MeshMergeUtilities/Private/MeshMergeHelpers.cpp b/Engine/Source/Developer/MeshMergeUtilities/Private/MeshMergeHelpers.cpp
--- a/Engine/Source/Developer/MeshMergeUtilities/Private/MeshMergeHelpers.cpp
+++ b/Engine/Source/Developer/MeshMergeUtilities/Private/MeshMergeHelpers.cpp
@@ -441,6 +441,8 @@ void FMeshMergeHelpers::ExportStaticMeshLOD(const FStaticMeshLODResources& Stati
 	}
 
 	FStaticMeshAttributes Attributes(OutMeshDescription);
+	Attributes.Register();
+
 	TVertexAttributesRef<FVector3f> VertexPositions = Attributes.GetVertexPositions();
 	TEdgeAttributesRef<bool> EdgeHardnesses = Attributes.GetEdgeHardnesses();
 	TPolygonGroupAttributesRef<FName> PolygonGroupImportedMaterialSlotNames = Attributes.GetPolygonGroupMaterialSlotNames();
