Ensure that the ISMC created from Merge Actors is transactional and supports per-instance editing

diff --git a/Engine/Source/Developer/MeshMergeUtilities/Private/MeshMergeUtilities.cpp b/Engine/Source/Developer/MeshMergeUtilities/Private/MeshMergeUtilities.cpp
--- a/Engine/Source/Developer/MeshMergeUtilities/Private/MeshMergeUtilities.cpp
+++ b/Engine/Source/Developer/MeshMergeUtilities/Private/MeshMergeUtilities.cpp
@@ -3537,7 +3537,8 @@ void FMeshMergeUtilities::MergeComponentsToInstances(const TArray<UPrimitiveComp
 
 						if (NewComponent == nullptr)
 						{
-							NewComponent = NewObject<UInstancedStaticMeshComponent>(ActorEntry.MergedActor, InSettings.ISMComponentToUse.Get());
+							NewComponent = NewObject<UInstancedStaticMeshComponent>(ActorEntry.MergedActor, InSettings.ISMComponentToUse.Get(), NAME_None, RF_Transactional);
+							NewComponent->bHasPerInstanceHitProxies = true;
 						
 							if (ActorEntry.MergedActor->GetRootComponent())
 							{
