Fixed static analysis warning

diff --git a/Engine/Source/Developer/MeshMergeUtilities/Private/MeshMergeUtilities.cpp b/Engine/Source/Developer/MeshMergeUtilities/Private/MeshMergeUtilities.cpp
--- a/Engine/Source/Developer/MeshMergeUtilities/Private/MeshMergeUtilities.cpp
+++ b/Engine/Source/Developer/MeshMergeUtilities/Private/MeshMergeUtilities.cpp
@@ -538,8 +538,11 @@ void FMeshMergeUtilities::ConvertOutputToFlatMaterials(const TArray<FBakeOutput>
 		for (TPair<EMaterialProperty, FIntPoint> SizePair : Output.PropertySizes)
 		{
 			EFlattenMaterialProperties OldProperty = ToFlattenProperty(SizePair.Key);
-			Material.SetPropertySize(OldProperty, SizePair.Value);
-			Material.GetPropertySamples(OldProperty).Append(Output.PropertyData[SizePair.Key]);
+			if (ensure(OldProperty != EFlattenMaterialProperties::NumFlattenMaterialProperties))
+			{
+				Material.SetPropertySize(OldProperty, SizePair.Value);
+				Material.GetPropertySamples(OldProperty).Append(Output.PropertyData[SizePair.Key]);
+			}
 		}
 
 		Material.bDitheredLODTransition = MaterialInfo.Material->IsDitheredLODTransition();
@@ -567,8 +570,11 @@ void FMeshMergeUtilities::TransferOutputToFlatMaterials(const TArray<FMaterialDa
 		for (TPair<EMaterialProperty, FIntPoint> SizePair : Output.PropertySizes)
 		{
 			EFlattenMaterialProperties OldProperty = ToFlattenProperty(SizePair.Key);
-			Material.SetPropertySize(OldProperty, SizePair.Value);
-			Material.GetPropertySamples(OldProperty) = MoveTemp(Output.PropertyData[SizePair.Key]);
+			if (ensure(OldProperty != EFlattenMaterialProperties::NumFlattenMaterialProperties))
+			{
+				Material.SetPropertySize(OldProperty, SizePair.Value);
+				Material.GetPropertySamples(OldProperty) = MoveTemp(Output.PropertyData[SizePair.Key]);
+			}
 		}
 
 		Material.bDitheredLODTransition = MaterialInfo.Material->IsDitheredLODTransition();
@@ -2649,7 +2655,10 @@ void FMeshMergeUtilities::MergeComponentsToStaticMesh(const TArray<UPrimitiveCom
 			if (Entry.Property != MP_MAX)
 			{
 				EFlattenMaterialProperties OldProperty = ToFlattenProperty(Entry.Property);
-				OutMaterial.SetPropertySize(OldProperty, Entry.bUseCustomSize ? Entry.CustomSize : MaterialOptions->TextureSize);
+				if (ensure(OldProperty != EFlattenMaterialProperties::NumFlattenMaterialProperties))
+				{
+					OutMaterial.SetPropertySize(OldProperty, Entry.bUseCustomSize ? Entry.CustomSize : MaterialOptions->TextureSize);
+				}
 			}
 		}
 
