Allow the use of alternate names for flattening material params

diff --git a/Engine/Source/Developer/MeshMergeUtilities/Private/MeshMergeUtilities.cpp b/Engine/Source/Developer/MeshMergeUtilities/Private/MeshMergeUtilities.cpp
--- a/Engine/Source/Developer/MeshMergeUtilities/Private/MeshMergeUtilities.cpp
+++ b/Engine/Source/Developer/MeshMergeUtilities/Private/MeshMergeUtilities.cpp
@@ -76,11 +76,6 @@
 
 #include "ISMPartition/ISMComponentDescriptor.h"
 
-#if WITH_EDITOR
-#include "Widgets/Notifications/SNotificationList.h"
-#include "Framework/Notifications/NotificationManager.h"
-#endif // WITH_EDITOR
-
 #define LOCTEXT_NAMESPACE "MeshMergeUtils"
 
 DEFINE_LOG_CATEGORY(LogMeshMerging);
@@ -1966,79 +1961,7 @@ void FMeshMergeUtilities::CreateProxyMesh(const TArray<UStaticMeshComponent*>& I
 
 bool FMeshMergeUtilities::IsValidBaseMaterial(const UMaterialInterface* InBaseMaterial, bool bShowToaster) const
 {
-	if (InBaseMaterial != nullptr)
-	{
-		TArray<FGuid> ParameterIds;
-		TArray<FString> MissingParameters;
-		auto NameCheckLambda = [&MissingParameters](const TArray<FMaterialParameterInfo>& InCheck, const TArray<FName>& InRequired)
-		{
-			for (const FName& Name : InRequired)
-			{
-				if (!InCheck.ContainsByPredicate([Name](const FMaterialParameterInfo& ParamInfo) { return (ParamInfo.Name == Name); }))
-				{
-					MissingParameters.Add(Name.ToString());
-				}
-			}
-		};
-
-		TArray<FMaterialParameterInfo> TextureParameterInfos;
-		TArray<FName> RequiredTextureNames = { TEXT("DiffuseTexture"), TEXT("NormalTexture"), TEXT("PackedTexture"), TEXT("MetallicTexture"), TEXT("SpecularTexture"), TEXT("RoughnessTexture"), TEXT("EmissiveTexture"), TEXT("OpacityTexture"), TEXT("OpacityMaskTexture"), TEXT("AmbientOcclusionTexture") };
-		InBaseMaterial->GetAllTextureParameterInfo(TextureParameterInfos, ParameterIds);
-		NameCheckLambda(TextureParameterInfos, RequiredTextureNames);
-
-		TArray<FMaterialParameterInfo> ScalarParameterInfos;
-		TArray<FName> RequiredScalarNames = { TEXT("MetallicConst"), TEXT("SpecularConst"), TEXT("RoughnessConst"), TEXT("OpacityConst"), TEXT("OpacityMaskConst"), TEXT("AmbientOcclusionConst"), TEXT("EmissiveScale") };
-		InBaseMaterial->GetAllScalarParameterInfo(ScalarParameterInfos, ParameterIds);
-		NameCheckLambda(ScalarParameterInfos, RequiredScalarNames);
-
-		TArray<FMaterialParameterInfo> VectorParameterInfos;
-		TArray<FName> RequiredVectorNames = { TEXT("DiffuseConst"), TEXT("EmissiveConst") };
-		InBaseMaterial->GetAllVectorParameterInfo(VectorParameterInfos, ParameterIds);
-		NameCheckLambda(VectorParameterInfos, RequiredVectorNames);
-
-		TArray<FMaterialParameterInfo> StaticSwitchParameterInfos;
-		TArray<FName> RequiredSwitchNames = { TEXT("UseDiffuse"), TEXT("PackMetallic"), TEXT("PackSpecular"), TEXT("PackRoughness"),TEXT("UseMetallic"), TEXT("UseSpecular"), TEXT("UseRoughness"), TEXT("UseEmissive"), TEXT("UseOpacity"), TEXT("UseOpacityMask"), TEXT("UseAmbientOcclusion") };
-		InBaseMaterial->GetAllStaticSwitchParameterInfo(StaticSwitchParameterInfos, ParameterIds);
-		NameCheckLambda(StaticSwitchParameterInfos, RequiredSwitchNames);
-
-		if (MissingParameters.Num() > 0)
-		{
-			FString MissingNamesString;
-			for (const FString& Name : MissingParameters)
-			{
-				if (!MissingNamesString.IsEmpty())
-				{
-					MissingNamesString += ", ";
-					MissingNamesString += Name;
-				}
-				else
-				{
-					MissingNamesString += Name;
-				}
-			}
-#if WITH_EDITOR
-			if (bShowToaster)
-			{
-				FFormatNamedArguments Arguments;
-				Arguments.Add(TEXT("MaterialName"), FText::FromString(InBaseMaterial->GetName()));
-				FText ErrorMessage = FText::Format(LOCTEXT("UHierarchicalLODSettings_PostEditChangeProperty", "Material {MaterialName} is missing required Material Parameters (check log for details)"), Arguments);
-				FNotificationInfo Info(ErrorMessage);
-				Info.ExpireDuration = 5.0f;
-				FSlateNotificationManager::Get().AddNotification(Info);
-			}
-
-			UE_LOG(LogMeshMerging, Error, TEXT("Material %s is missing required Material Parameters %s, resetting to default."), *InBaseMaterial->GetName(), *MissingNamesString);
-#endif // WITH_EDITOR
-
-			return false;
-		}
-		else
-		{
-			return true;
-		}
-	}
-
-	return false;
+	return ProxyMaterialUtilities::IsValidBaseMaterial(InBaseMaterial, bShowToaster);
 }
 
 void FMeshMergeUtilities::RetrieveMeshDescription(const UStaticMeshComponent* InStaticMeshComponent, int32 LODIndex, FMeshDescription& InOutMeshDescription, bool bPropagateMeshData) const
diff --git a/Engine/Source/Developer/MeshMergeUtilities/Private/ProxyMaterialUtilities.h b/Engine/Source/Developer/MeshMergeUtilities/Private/ProxyMaterialUtilities.h
--- a/Engine/Source/Developer/MeshMergeUtilities/Private/ProxyMaterialUtilities.h
+++ b/Engine/Source/Developer/MeshMergeUtilities/Private/ProxyMaterialUtilities.h
@@ -13,8 +13,217 @@
 #include "MeshMergeModule.h"
 #include "Modules/ModuleManager.h"
 
+#if WITH_EDITOR
+#include "Widgets/Notifications/SNotificationList.h"
+#include "Framework/Notifications/NotificationManager.h"
+#endif // WITH_EDITOR
+
+#define LOCTEXT_NAMESPACE "ProxyMaterialUtilities"
+
 namespace ProxyMaterialUtilities
 {
+	// Find the parameter name to use with the provided material, for the given property
+	static TArray<FString> GetPotentialParamNames(EFlattenMaterialProperties InProperty)
+	{
+		switch(InProperty)
+		{
+			case EFlattenMaterialProperties::Diffuse: return { "BaseColor", "Diffuse" };
+			case EFlattenMaterialProperties::Normal: return { "Normal" };
+			case EFlattenMaterialProperties::Metallic: return { "Metallic" };
+			case EFlattenMaterialProperties::Roughness: return { "Roughness" };
+			case EFlattenMaterialProperties::Specular: return { "Specular" };
+			case EFlattenMaterialProperties::Opacity: return { "Opacity" };
+			case EFlattenMaterialProperties::OpacityMask: return { "OpacityMask" };
+			case EFlattenMaterialProperties::AmbientOcclusion: return { "AmbientOcclusion" };
+			case EFlattenMaterialProperties::Emissive: return { "EmissiveColor", "Emissive" };
+		}
+
+		return TArray<FString>();
+	}
+
+	static EMaterialParameterType GetConstantParamType(EFlattenMaterialProperties InProperty)
+	{
+		EMaterialParameterType ParamType = EMaterialParameterType::Count;
+
+		switch (InProperty)
+		{
+		case EFlattenMaterialProperties::Metallic:
+		case EFlattenMaterialProperties::Roughness:
+		case EFlattenMaterialProperties::Specular:
+		case EFlattenMaterialProperties::Opacity:
+		case EFlattenMaterialProperties::OpacityMask:
+		case EFlattenMaterialProperties::AmbientOcclusion:
+			ParamType = EMaterialParameterType::Scalar;
+			break;
+
+		case EFlattenMaterialProperties::Diffuse:
+		case EFlattenMaterialProperties::Emissive:
+			ParamType = EMaterialParameterType::Vector;
+			break;
+		}
+
+		return ParamType;
+	}	
+
+    // Find the parameter name to use with the provided material, for the given property
+	static bool GetMatchingParamName(EFlattenMaterialProperties InProperty, const UMaterialInterface* InBaseMaterial, FString& OutParamName, TArray<FString>* OutMissingNames = nullptr)
+	{
+		const TArray<FString> PotentialNames = GetPotentialParamNames(InProperty);
+
+		// Missing names, for error reporting
+		FString MissingNames;
+
+		for (const FString& PotentialName : PotentialNames)
+		{
+			const FName TextureName(PotentialName + TEXT("Texture"));
+			const FName ConstName(PotentialName + TEXT("Const"));
+			const FName UseTexture(TEXT("Use") + PotentialName);
+
+			UTexture* DefaultTexture = nullptr;
+			bool DefaultSwitchValue = false;
+			float DefaultScalarValue;
+			FLinearColor DefaultVectorValue;
+			FGuid ExpressionGuid;
+
+			if (!MissingNames.IsEmpty())
+			{
+				MissingNames += TEXT("|");
+			}
+
+			bool bHasRequiredParams = InBaseMaterial->GetTextureParameterValue(TextureName, DefaultTexture);
+
+			switch(GetConstantParamType(InProperty))
+			{
+			case EMaterialParameterType::Scalar:
+				MissingNames += UseTexture.ToString() + "+" + TextureName.ToString() + "+" + ConstName.ToString();
+				bHasRequiredParams &= InBaseMaterial->GetStaticSwitchParameterDefaultValue(UseTexture, DefaultSwitchValue, ExpressionGuid) && 
+									  InBaseMaterial->GetScalarParameterDefaultValue(ConstName, DefaultScalarValue);
+				break;
+
+			case EMaterialParameterType::Vector:
+				MissingNames += UseTexture.ToString() + "+" + TextureName.ToString() + "+" + ConstName.ToString();
+				bHasRequiredParams &= InBaseMaterial->GetStaticSwitchParameterDefaultValue(UseTexture, DefaultSwitchValue, ExpressionGuid) &&
+									  InBaseMaterial->GetVectorParameterDefaultValue(ConstName, DefaultVectorValue);
+				break;
+
+			case EMaterialParameterType::Count:
+				MissingNames += UseTexture.ToString();
+				break;
+			}
+
+			if (bHasRequiredParams)
+			{
+				OutParamName = PotentialName;
+				return true;
+			}
+		}
+
+		if (OutMissingNames != nullptr)
+		{
+			OutMissingNames->Add(MissingNames);
+		}
+
+		OutParamName = "";
+		return false;
+	}
+
+	static FString GetMatchingParamName(EFlattenMaterialProperties InProperty, UMaterialInterface* InBaseMaterial)
+	{
+		FString ParamName = "";
+		if (!GetMatchingParamName(InProperty, InBaseMaterial, ParamName))
+		{
+			UE_LOG(LogMeshMerging, Fatal, TEXT("Invalid base material, should have been rejected by IsValidBaseMaterial()"));
+		}
+		return ParamName;
+	}
+
+	// Validate that the provided material has all the required parameters needed to be considered a flattening material
+	// If not, report what is missing
+	static bool IsValidBaseMaterial(const UMaterialInterface* InBaseMaterial, bool bShowToaster)
+	{
+		if (InBaseMaterial != nullptr)
+		{
+			TArray<FGuid> ParameterIds;
+			TArray<FString> MissingParameters;
+
+
+			FString ParamName;
+			GetMatchingParamName(EFlattenMaterialProperties::Diffuse, InBaseMaterial, ParamName, &MissingParameters);
+			GetMatchingParamName(EFlattenMaterialProperties::Normal, InBaseMaterial, ParamName, &MissingParameters);
+			GetMatchingParamName(EFlattenMaterialProperties::Metallic, InBaseMaterial, ParamName, &MissingParameters);
+			GetMatchingParamName(EFlattenMaterialProperties::Roughness, InBaseMaterial, ParamName, &MissingParameters);
+			GetMatchingParamName(EFlattenMaterialProperties::Specular, InBaseMaterial, ParamName, &MissingParameters);
+			GetMatchingParamName(EFlattenMaterialProperties::Opacity, InBaseMaterial, ParamName, &MissingParameters);
+			GetMatchingParamName(EFlattenMaterialProperties::OpacityMask, InBaseMaterial, ParamName, &MissingParameters);
+			GetMatchingParamName(EFlattenMaterialProperties::AmbientOcclusion, InBaseMaterial, ParamName, &MissingParameters);
+			GetMatchingParamName(EFlattenMaterialProperties::Emissive, InBaseMaterial, ParamName, &MissingParameters);
+
+			auto NameCheckLambda = [&MissingParameters](const TArray<FMaterialParameterInfo>& InCheck, const TArray<FName>& InRequired)
+			{
+				for (const FName& Name : InRequired)
+				{
+					if (!InCheck.ContainsByPredicate([Name](const FMaterialParameterInfo& ParamInfo) { return (ParamInfo.Name == Name); }))
+					{
+						MissingParameters.Add(Name.ToString());
+					}
+				}
+			};
+
+			TArray<FMaterialParameterInfo> TextureParameterInfos;
+			TArray<FName> RequiredTextureNames = { TEXT("PackedTexture") };
+			InBaseMaterial->GetAllTextureParameterInfo(TextureParameterInfos, ParameterIds);
+			NameCheckLambda(TextureParameterInfos, RequiredTextureNames);
+
+			TArray<FMaterialParameterInfo> ScalarParameterInfos;
+			TArray<FName> RequiredScalarNames = { TEXT("EmissiveScale") };
+			InBaseMaterial->GetAllScalarParameterInfo(ScalarParameterInfos, ParameterIds);
+			NameCheckLambda(ScalarParameterInfos, RequiredScalarNames);
+
+			TArray<FMaterialParameterInfo> StaticSwitchParameterInfos;
+			TArray<FName> RequiredSwitchNames = { TEXT("PackMetallic"), TEXT("PackSpecular"), TEXT("PackRoughness") };
+			InBaseMaterial->GetAllStaticSwitchParameterInfo(StaticSwitchParameterInfos, ParameterIds);
+			NameCheckLambda(StaticSwitchParameterInfos, RequiredSwitchNames);
+
+			if (MissingParameters.Num() > 0)
+			{
+				FString MissingNamesString;
+				for (const FString& Name : MissingParameters)
+				{
+					if (!MissingNamesString.IsEmpty())
+					{
+						MissingNamesString += ", ";
+						MissingNamesString += Name;
+					}
+					else
+					{
+						MissingNamesString += Name;
+					}
+				}
+#if WITH_EDITOR
+				if (bShowToaster)
+				{
+					FFormatNamedArguments Arguments;
+					Arguments.Add(TEXT("MaterialName"), FText::FromString(InBaseMaterial->GetName()));
+					FText ErrorMessage = FText::Format(LOCTEXT("UHierarchicalLODSettings_PostEditChangeProperty", "Material {MaterialName} is missing required Material Parameters (check log for details)"), Arguments);
+					FNotificationInfo Info(ErrorMessage);
+					Info.ExpireDuration = 5.0f;
+					FSlateNotificationManager::Get().AddNotification(Info);
+				}
+
+				UE_LOG(LogMeshMerging, Error, TEXT("Material %s is missing required Material Parameters %s, resetting to default."), *InBaseMaterial->GetName(), *MissingNamesString);
+#endif // WITH_EDITOR
+
+				return false;
+			}
+			else
+			{
+				return true;
+			}
+		}
+
+		return false;
+	}
+
 	static const bool CalculatePackedTextureData(const FFlattenMaterial& InMaterial, bool& bOutPackMetallic, bool& bOutPackSpecular, bool& bOutPackRoughness, int32& OutNumSamples, FIntPoint& OutSize)
 	{
 		// Whether or not a material property is baked down
@@ -135,10 +344,12 @@ namespace ProxyMaterialUtilities
 			return Texture;
 		};
 
-		auto SetTextureParam = [&](const FString& PropertyName, EFlattenMaterialProperties FlattenProperty)
+		auto SetTextureParam = [&](EFlattenMaterialProperties FlattenProperty)
 		{
 			if (FlattenMaterial.DoesPropertyContainData(FlattenProperty) && !FlattenMaterial.IsPropertyConstant(FlattenProperty))
 			{
+				FString PropertyName = GetMatchingParamName(FlattenProperty, OutMaterial);
+
 				const FName TextureName(PropertyName + TEXT("Texture"));
 				const FName UseTexture(TEXT("Use") + PropertyName);
 
@@ -156,27 +367,31 @@ namespace ProxyMaterialUtilities
 			}
 		};
 
-		auto SetTextureParamConstVector = [&](const FString& PropertyName, EFlattenMaterialProperties FlattenProperty)
+		auto SetTextureParamConstVector = [&](EFlattenMaterialProperties FlattenProperty)
 		{
 			if (FlattenMaterial.DoesPropertyContainData(FlattenProperty) && !FlattenMaterial.IsPropertyConstant(FlattenProperty))
 			{
-				SetTextureParam(PropertyName, FlattenProperty);
+				SetTextureParam(FlattenProperty);
 			} 
 			else
 			{
+				FString PropertyName = GetMatchingParamName(FlattenProperty, OutMaterial);
+
 				const FName ConstName(PropertyName + TEXT("Const"));
 				OutMaterial->SetVectorParameterValueEditorOnly(ConstName, FlattenMaterial.GetPropertySamples(FlattenProperty)[0]);
 			} 
 		};
 
-		auto SetTextureParamConstScalar = [&](const FString& PropertyName, EFlattenMaterialProperties FlattenProperty, float ConstantValue)
+		auto SetTextureParamConstScalar = [&](EFlattenMaterialProperties FlattenProperty, float ConstantValue)
 		{
 			if (FlattenMaterial.DoesPropertyContainData(FlattenProperty) && !FlattenMaterial.IsPropertyConstant(FlattenProperty))
 			{
-				SetTextureParam(PropertyName, FlattenProperty);
+				SetTextureParam(FlattenProperty);
 			}
 			else
 			{
+				FString PropertyName = GetMatchingParamName(FlattenProperty, OutMaterial);
+
 				const FName ConstName(PropertyName + TEXT("Const"));
 				FLinearColor Colour = FlattenMaterial.IsPropertyConstant(FlattenProperty) ? FLinearColor::FromSRGBColor(FlattenMaterial.GetPropertySamples(FlattenProperty)[0]) : FLinearColor(ConstantValue, 0, 0, 0);
 				OutMaterial->SetScalarParameterValueEditorOnly(ConstName, Colour.R);
@@ -186,43 +401,43 @@ namespace ProxyMaterialUtilities
 		// Load textures and set switches accordingly
 		if (FlattenMaterial.GetPropertySize(EFlattenMaterialProperties::Diffuse).Num() > 0 && !(FlattenMaterial.IsPropertyConstant(EFlattenMaterialProperties::Diffuse) && FlattenMaterial.GetPropertySamples(EFlattenMaterialProperties::Diffuse)[0] == FColor::Black))
 		{
-			SetTextureParamConstVector("BaseColor", EFlattenMaterialProperties::Diffuse);
+			SetTextureParamConstVector(EFlattenMaterialProperties::Diffuse);
 		}
 
 		if (FlattenMaterial.GetPropertySize(EFlattenMaterialProperties::Normal).Num() > 1)
 		{
-			SetTextureParam("Normal", EFlattenMaterialProperties::Normal);
+			SetTextureParam(EFlattenMaterialProperties::Normal);
 		}
 
 		// Determine whether or not specific material properties are packed together into one texture (requires at least two to match (number of samples and texture size) in order to be packed
 		if (!bPackMetallic && (FlattenMaterial.GetPropertySize(EFlattenMaterialProperties::Metallic).Num() > 0 || !InMaterialProxySettings.bMetallicMap))
 		{
-			SetTextureParamConstScalar("Metallic", EFlattenMaterialProperties::Metallic, InMaterialProxySettings.MetallicConstant);
+			SetTextureParamConstScalar(EFlattenMaterialProperties::Metallic, InMaterialProxySettings.MetallicConstant);
 		}
 
 		if (!bPackRoughness && (FlattenMaterial.GetPropertySize(EFlattenMaterialProperties::Roughness).Num() > 0 || !InMaterialProxySettings.bRoughnessMap))
 		{
-			SetTextureParamConstScalar("Roughness", EFlattenMaterialProperties::Roughness, InMaterialProxySettings.RoughnessConstant);
+			SetTextureParamConstScalar(EFlattenMaterialProperties::Roughness, InMaterialProxySettings.RoughnessConstant);
 		}
 
 		if (!bPackSpecular && (FlattenMaterial.GetPropertySize(EFlattenMaterialProperties::Specular).Num() > 0 || !InMaterialProxySettings.bSpecularMap))
 		{
-			SetTextureParamConstScalar("Specular", EFlattenMaterialProperties::Specular, InMaterialProxySettings.SpecularConstant);
+			SetTextureParamConstScalar(EFlattenMaterialProperties::Specular, InMaterialProxySettings.SpecularConstant);
 		}
 
 		if (FlattenMaterial.GetPropertySize(EFlattenMaterialProperties::Opacity).Num() > 0 || !InMaterialProxySettings.bOpacityMap)
 		{
-			SetTextureParamConstScalar("Opacity", EFlattenMaterialProperties::Opacity, InMaterialProxySettings.OpacityConstant);
+			SetTextureParamConstScalar(EFlattenMaterialProperties::Opacity, InMaterialProxySettings.OpacityConstant);
 		}
 
 		if (FlattenMaterial.GetPropertySize(EFlattenMaterialProperties::OpacityMask).Num() > 0 || !InMaterialProxySettings.bOpacityMaskMap)
 		{
-			SetTextureParamConstScalar("OpacityMask", EFlattenMaterialProperties::OpacityMask, InMaterialProxySettings.OpacityMaskConstant);
+			SetTextureParamConstScalar(EFlattenMaterialProperties::OpacityMask, InMaterialProxySettings.OpacityMaskConstant);
 		}
 
 		if (FlattenMaterial.GetPropertySize(EFlattenMaterialProperties::AmbientOcclusion).Num() > 0 || !InMaterialProxySettings.bAmbientOcclusionMap)
 		{
-			SetTextureParamConstScalar("AmbientOcclusion", EFlattenMaterialProperties::AmbientOcclusion, InMaterialProxySettings.AmbientOcclusionConstant);
+			SetTextureParamConstScalar(EFlattenMaterialProperties::AmbientOcclusion, InMaterialProxySettings.AmbientOcclusionConstant);
 		}
 
 		// Handle the packed texture if applicable
@@ -289,7 +504,7 @@ namespace ProxyMaterialUtilities
 		// Emissive is a special case due to the scaling variable
 		if (FlattenMaterial.GetPropertySamples(EFlattenMaterialProperties::Emissive).Num() > 0 && !(FlattenMaterial.GetPropertySize(EFlattenMaterialProperties::Emissive).Num() == 1 && FlattenMaterial.GetPropertySamples(EFlattenMaterialProperties::Emissive)[0] == FColor::Black))
 		{
-			SetTextureParamConstVector("EmissiveColor", EFlattenMaterialProperties::Emissive);
+			SetTextureParamConstVector(EFlattenMaterialProperties::Emissive);
 
 			if (FlattenMaterial.EmissiveScale != 1.0f)
 			{
@@ -313,3 +528,5 @@ namespace ProxyMaterialUtilities
 	}
 
 };
+
+#undef LOCTEXT_NAMESPACE // "ProxyMaterialUtilities"
\ No newline at end of file
