Fixed PolyGroup issue with FMeshMergeHelpers::MergeImpostersToMesh()

The function would remap polygroups of the source meshes, then use  FStaticMeshOperations::AppendMeshDescription() which would also mess with the groups.

Providing a polygroup delegate to AppendMeshDescription() solves the issue.

diff --git a/Engine/Source/Developer/MeshMergeUtilities/Private/MeshMergeHelpers.cpp b/Engine/Source/Developer/MeshMergeUtilities/Private/MeshMergeHelpers.cpp
--- a/Engine/Source/Developer/MeshMergeUtilities/Private/MeshMergeHelpers.cpp
+++ b/Engine/Source/Developer/MeshMergeUtilities/Private/MeshMergeHelpers.cpp
@@ -1203,28 +1203,30 @@ void FMeshMergeHelpers::MergeImpostersToMesh(TArray<const UStaticMeshComponent*>
 		FStaticMeshAttributes Attributes(InMeshDescription);
 		TPolygonGroupAttributesRef<FName> TargetPolygonGroupImportedMaterialSlotNames = Attributes.GetPolygonGroupMaterialSlotNames();
 
+		FStaticMeshOperations::FAppendSettings AppendSettings;
+
 		//Add the missing polygon group ID to the target(InMeshDescription)
 		//Remap the source mesh(ImposterMesh) polygongroup to fit with the target polygon groups
-		TMap<FPolygonGroupID, FPolygonGroupID> RemapSourcePolygonGroup;
-		RemapSourcePolygonGroup.Reserve(ImposterMesh.PolygonGroups().Num());
-		int32 SectionIndex = 0;
-		for (const FPolygonGroupID SourcePolygonGroupID : ImposterMesh.PolygonGroups().GetElementIDs())
+		AppendSettings.PolygonGroupsDelegate = FAppendPolygonGroupsDelegate::CreateLambda([&](const FMeshDescription& SourceMesh, FMeshDescription& TargetMesh, PolygonGroupMap& RemapPolygonGroups)
 		{
-			UMaterialInterface* MaterialUseBySection = OutImposterMaterials[SectionImposterUniqueMaterialIndex[SectionIndex++]];
-			FPolygonGroupID* ExistTargetPolygonGroupID = ImposterMaterialToPolygonGroupID.Find(MaterialUseBySection);
-			FPolygonGroupID MatchTargetPolygonGroupID = ExistTargetPolygonGroupID == nullptr ? INDEX_NONE : *ExistTargetPolygonGroupID;
-			if (MatchTargetPolygonGroupID == INDEX_NONE)
+			RemapPolygonGroups.Reserve(SourceMesh.PolygonGroups().Num());
+			int32 SectionIndex = 0;
+			for (const FPolygonGroupID SourcePolygonGroupID : SourceMesh.PolygonGroups().GetElementIDs())
 			{
-				MatchTargetPolygonGroupID = InMeshDescription.CreatePolygonGroup();
-				//use the material name to fill the imported material name. Material name will be unique
-				TargetPolygonGroupImportedMaterialSlotNames[MatchTargetPolygonGroupID] = MaterialUseBySection->GetFName();
-				ImposterMaterialToPolygonGroupID.Add(MaterialUseBySection, MatchTargetPolygonGroupID);
+				UMaterialInterface* MaterialUseBySection = OutImposterMaterials[SectionImposterUniqueMaterialIndex[SectionIndex++]];
+				FPolygonGroupID* ExistTargetPolygonGroupID = ImposterMaterialToPolygonGroupID.Find(MaterialUseBySection);
+				FPolygonGroupID MatchTargetPolygonGroupID = ExistTargetPolygonGroupID == nullptr ? INDEX_NONE : *ExistTargetPolygonGroupID;
+				if (MatchTargetPolygonGroupID == INDEX_NONE)
+				{
+					MatchTargetPolygonGroupID = TargetMesh.CreatePolygonGroup();
+					//use the material name to fill the imported material name. Material name will be unique
+					TargetPolygonGroupImportedMaterialSlotNames[MatchTargetPolygonGroupID] = MaterialUseBySection->GetFName();
+					ImposterMaterialToPolygonGroupID.Add(MaterialUseBySection, MatchTargetPolygonGroupID);
+				}
+				RemapPolygonGroups.Add(SourcePolygonGroupID, MatchTargetPolygonGroupID);
 			}
-			RemapSourcePolygonGroup.Add(SourcePolygonGroupID, MatchTargetPolygonGroupID);
-		}
-		ImposterMesh.RemapPolygonGroups(RemapSourcePolygonGroup);
+		});
 
-		FStaticMeshOperations::FAppendSettings AppendSettings;
 		AppendSettings.bMergeVertexColor = true;
 		for (int32 ChannelIdx = 0; ChannelIdx < FStaticMeshOperations::FAppendSettings::MAX_NUM_UV_CHANNELS; ++ChannelIdx)
 		{
