MergeActors: workaround for crash where unique material list is larger than number of unique sections
#rnx

diff --git a/Engine/Source/Developer/MeshMergeUtilities/Private/MeshMergeUtilities.cpp b/Engine/Source/Developer/MeshMergeUtilities/Private/MeshMergeUtilities.cpp
--- a/Engine/Source/Developer/MeshMergeUtilities/Private/MeshMergeUtilities.cpp
+++ b/Engine/Source/Developer/MeshMergeUtilities/Private/MeshMergeUtilities.cpp
@@ -2808,12 +2808,16 @@ void FMeshMergeUtilities::MergeComponentsToStaticMesh(const TArray<UPrimitiveCom
 		for (int32 Index = 0; Index < UniqueMaterialIndices.Num(); ++Index)
 		{
 			const int32 SectionIndex = UniqueMaterialIndices[Index];
-			const FSectionInfo& StoredSectionInfo = DataTracker.GetSection(SectionIndex);
-			FMeshSectionInfo SectionInfo;
-			SectionInfo.bCastShadow = StoredSectionInfo.EnabledProperties.Contains(GET_MEMBER_NAME_CHECKED(FMeshSectionInfo, bCastShadow));
-			SectionInfo.bEnableCollision = StoredSectionInfo.EnabledProperties.Contains(GET_MEMBER_NAME_CHECKED(FMeshSectionInfo, bEnableCollision));
-			SectionInfo.MaterialIndex = UniqueMaterials.IndexOfByKey(StoredSectionInfo.Material);
-			SectionInfoMap.Set(LODIndex, Index, SectionInfo);
+			// unclear when this would not be the case, but it seems to be able to occur
+			if (SectionIndex < DataTracker.NumberOfUniqueSections())
+			{
+				const FSectionInfo& StoredSectionInfo = DataTracker.GetSection(SectionIndex);
+				FMeshSectionInfo SectionInfo;
+				SectionInfo.bCastShadow = StoredSectionInfo.EnabledProperties.Contains(GET_MEMBER_NAME_CHECKED(FMeshSectionInfo, bCastShadow));
+				SectionInfo.bEnableCollision = StoredSectionInfo.EnabledProperties.Contains(GET_MEMBER_NAME_CHECKED(FMeshSectionInfo, bEnableCollision));
+				SectionInfo.MaterialIndex = UniqueMaterials.IndexOfByKey(StoredSectionInfo.Material);
+				SectionInfoMap.Set(LODIndex, Index, SectionInfo);
+			}
 		}
 	}
 
