Support building HLODs meshes using cooked data
* Tested mesh merge/simplify/approximate methods
* Fixed issue in Mesh Approximate where import build scale would be applied twice (as LODRenderData is already scaled)

diff --git a/Engine/Plugins/Runtime/MeshModelingToolset/Source/ModelingComponents/Private/Scene/MeshSceneAdapter.cpp b/Engine/Plugins/Runtime/MeshModelingToolset/Source/ModelingComponents/Private/Scene/MeshSceneAdapter.cpp
--- a/Engine/Plugins/Runtime/MeshModelingToolset/Source/ModelingComponents/Private/Scene/MeshSceneAdapter.cpp
+++ b/Engine/Plugins/Runtime/MeshModelingToolset/Source/ModelingComponents/Private/Scene/MeshSceneAdapter.cpp
@@ -925,7 +925,6 @@ public:
 	UStaticMesh* StaticMesh = nullptr;
 	int32 LODIndex = 0;
 
-	FVector3d BuildScale = FVector3d::One();
 	TUniquePtr<MeshAdapterType> Adapter;
 
 	TUniquePtr<TMeshAABBTree3<MeshAdapterType>> AABBTree;
@@ -947,14 +946,6 @@ public:
 			return false;
 		}
 
-		BuildScale = FVector3d::One();
-#if WITH_EDITOR
-		// respect BuildScale build setting
-		const FMeshBuildSettings& LODBuildSettings = StaticMesh->GetSourceModel(LODIndex).BuildSettings;
-		BuildScale = (FVector3d)LODBuildSettings.BuildScale3D;
-		Adapter->SetBuildScale(BuildScale, false);
-#endif
-
 		if (BuildOptions.bBuildSpatialDataStructures)
 		{
 			{
@@ -1139,7 +1130,10 @@ public:
 	{
 #if WITH_EDITOR
 		CachedSourceMeshDescription = StaticMesh->GetMeshDescription(LODIndex);
-		return MakeUnique<FMeshDescriptionTriangleMeshSurfaceAdapter>(CachedSourceMeshDescription, StaticMesh, BuildOptions.bOnlySurfaceMaterials);
+		TUniquePtr<FMeshDescriptionTriangleMeshSurfaceAdapter> SurfaceAdapter = MakeUnique<FMeshDescriptionTriangleMeshSurfaceAdapter>(CachedSourceMeshDescription, StaticMesh, BuildOptions.bOnlySurfaceMaterials);
+		const FMeshBuildSettings& LODBuildSettings = StaticMesh->GetSourceModel(LODIndex).BuildSettings;
+		SurfaceAdapter->SetBuildScale(LODBuildSettings.BuildScale3D, false);
+		return SurfaceAdapter;
 #else
 		// cannot use this path in non-editor builds, but this should have been handled higher up by instantiating 
 		// the runtime subclass
@@ -1315,46 +1309,39 @@ public:
 
 static TUniquePtr<IMeshSpatialWrapper> SpatialWrapperFactory( const FMeshTypeContainer& MeshContainer, const FMeshSceneAdapterBuildOptions& BuildOptions )
 {
-
 	if (MeshContainer.MeshType == ESceneMeshType::StaticMeshAsset)
 	{
+		UStaticMesh* StaticMesh = MeshContainer.GetStaticMesh();
+		if (ensure(StaticMesh != nullptr))
+		{
 #if WITH_EDITOR
-		bool bUseRenderMeshes = BuildOptions.bIgnoreStaticMeshSourceData;
+			bool bUseRenderMeshes = BuildOptions.bIgnoreStaticMeshSourceData || StaticMesh->GetOutermost()->bIsCookedForEditor;
 #else
-		bool bUseRenderMeshes = true;
+			bool bUseRenderMeshes = true;
 #endif
 
-		if (bUseRenderMeshes)
-		{
-			TUniquePtr<FStaticMeshRenderDataSpatialWrapper> SMWrapper = MakeUnique<FStaticMeshRenderDataSpatialWrapper>();
-			SMWrapper->SourceContainer = MeshContainer;
-			SMWrapper->StaticMesh = MeshContainer.GetStaticMesh();
-			if (ensure(SMWrapper->StaticMesh != nullptr))
+			if (bUseRenderMeshes)
 			{
+				TUniquePtr<FStaticMeshRenderDataSpatialWrapper> SMWrapper = MakeUnique<FStaticMeshRenderDataSpatialWrapper>();
+				SMWrapper->SourceContainer = MeshContainer;
+				SMWrapper->StaticMesh = StaticMesh;
 				return SMWrapper;
 			}
-		}
-		else if (BuildOptions.bEnableUVQueries == false && BuildOptions.bEnableNormalsQueries == false)
-		{
-			TUniquePtr<FCompressedStaticMeshSpatialWrapper> SMWrapper = MakeUnique<FCompressedStaticMeshSpatialWrapper>();
-			SMWrapper->SourceContainer = MeshContainer;
-			SMWrapper->StaticMesh = MeshContainer.GetStaticMesh();
-			if (ensure(SMWrapper->StaticMesh != nullptr))
+			else if (BuildOptions.bEnableUVQueries == false && BuildOptions.bEnableNormalsQueries == false)
 			{
+				TUniquePtr<FCompressedStaticMeshSpatialWrapper> SMWrapper = MakeUnique<FCompressedStaticMeshSpatialWrapper>();
+				SMWrapper->SourceContainer = MeshContainer;
+				SMWrapper->StaticMesh = StaticMesh;
 				return SMWrapper;
 			}
-		}
-		else
-		{
-			TUniquePtr<FStaticMeshSourceDataSpatialWrapper> SMWrapper = MakeUnique<FStaticMeshSourceDataSpatialWrapper>();
-			SMWrapper->SourceContainer = MeshContainer;
-			SMWrapper->StaticMesh = MeshContainer.GetStaticMesh();
-			if (ensure(SMWrapper->StaticMesh != nullptr))
+			else
 			{
+				TUniquePtr<FStaticMeshSourceDataSpatialWrapper> SMWrapper = MakeUnique<FStaticMeshSourceDataSpatialWrapper>();
+				SMWrapper->SourceContainer = MeshContainer;
+				SMWrapper->StaticMesh = StaticMesh;
 				return SMWrapper;
 			}
 		}
-
 	}
 
 	return TUniquePtr<IMeshSpatialWrapper>();
diff --git a/Engine/Source/Developer/MeshMergeUtilities/Private/MeshMergeHelpers.cpp b/Engine/Source/Developer/MeshMergeUtilities/Private/MeshMergeHelpers.cpp
--- a/Engine/Source/Developer/MeshMergeUtilities/Private/MeshMergeHelpers.cpp
+++ b/Engine/Source/Developer/MeshMergeUtilities/Private/MeshMergeHelpers.cpp
@@ -227,18 +227,15 @@
 	InOutMeshDescription = CombinedMeshDescription;
 }
 
-void FMeshMergeHelpers::RetrieveMesh(const UStaticMeshComponent* StaticMeshComponent, int32 LODIndex, FMeshDescription& OutMeshDescription, bool bPropagateVertexColours)
+
+static void RetrieveMeshInternal(const UStaticMesh* StaticMesh, const UStaticMeshComponent* StaticMeshComponent, int32 LODIndex, FMeshDescription& OutMeshDescription, bool bPropagateVertexColours)
 {
-	TRACE_CPUPROFILER_EVENT_SCOPE(FMeshMergeHelpers::RetrieveMesh)
+	TRACE_CPUPROFILER_EVENT_SCOPE(FMeshMergeHelpers::RetrieveMeshInternal)
 
-	const UStaticMesh* StaticMesh = StaticMeshComponent->GetStaticMesh();
-	const FStaticMeshSourceModel& StaticMeshModel = StaticMesh->GetSourceModel(LODIndex);
+	check(StaticMesh);
 
-	// Imported meshes will have a valid mesh description
-	const bool bImportedMesh = StaticMesh->IsMeshDescriptionValid(LODIndex);
-
 	// Export the mesh data using static mesh render data
-	ExportStaticMeshLOD(StaticMesh->GetRenderData()->LODResources[LODIndex], OutMeshDescription, StaticMesh->GetStaticMaterials());
+	FMeshMergeHelpers::ExportStaticMeshLOD(StaticMesh->GetRenderData()->LODResources[LODIndex], OutMeshDescription, StaticMesh->GetStaticMaterials());
 
 	// Make sure the mesh is not irreparably malformed.
 	if (OutMeshDescription.VertexInstances().Num() <= 0)
@@ -246,183 +243,45 @@
 		return;
 	}
 
-	// Transform mesh to world space
-	FTransform ComponentToWorldTransform = StaticMeshComponent->GetComponentTransform();
+	// If we have a component, use it to retrieve transform & vertex colors (if requested)
+	if (StaticMeshComponent)
+	{
+		// Transform mesh to world space
+		FTransform ComponentToWorldTransform = StaticMeshComponent->GetComponentTransform();
 
-	// Handle spline mesh deformation
-	const bool bIsSplineMeshComponent = StaticMeshComponent->IsA<USplineMeshComponent>();
-	if (bIsSplineMeshComponent)
-	{
-		const USplineMeshComponent* SplineMeshComponent = Cast<USplineMeshComponent>(StaticMeshComponent);
-		// Deform mesh data according to the Spline Mesh Component's data
-		PropagateSplineDeformationToMesh(SplineMeshComponent, OutMeshDescription);
+		// Handle spline mesh deformation
+		const bool bIsSplineMeshComponent = StaticMeshComponent->IsA<USplineMeshComponent>();
+		if (bIsSplineMeshComponent)
+		{
+			const USplineMeshComponent* SplineMeshComponent = Cast<USplineMeshComponent>(StaticMeshComponent);
+			// Deform mesh data according to the Spline Mesh Component's data
+			FMeshMergeHelpers::PropagateSplineDeformationToMesh(SplineMeshComponent, OutMeshDescription);
+		}
+
+		// If specified propagate painted vertex colors into our raw mesh
+		if (bPropagateVertexColours)
+		{
+			FMeshMergeHelpers::PropagatePaintedColorsToMesh(StaticMeshComponent, LODIndex, OutMeshDescription);
+		}
+
+		// Transform raw mesh vertex data by the Static Mesh Component's component to world transformation	
+		FStaticMeshOperations::ApplyTransform(OutMeshDescription, ComponentToWorldTransform);
 	}
 
-	// If specified propagate painted vertex colors into our raw mesh
-	if (bPropagateVertexColours)
-	{
-		PropagatePaintedColorsToMesh(StaticMeshComponent, LODIndex, OutMeshDescription);
+	FMeshBuildSettings BuildSettings;
+	
+	// If editor data is not available, we won't have access to source model
+	const bool bHasSourceModels = StaticMesh->IsSourceModelValid(0);
+	if (bHasSourceModels)
+	{ 
+		// Imported meshes will have a valid mesh description
+		const bool bImportedMesh = bHasSourceModels ? false : StaticMesh->IsMeshDescriptionValid(LODIndex);
+
+		// Use build settings from base mesh for LOD entries that were generated inside Editor.
+		BuildSettings = bImportedMesh ? StaticMesh->GetSourceModel(LODIndex).BuildSettings : StaticMesh->GetSourceModel(0).BuildSettings;
 	}
 
-	// Transform raw mesh vertex data by the Static Mesh Component's component to world transformation	
-	FStaticMeshOperations::ApplyTransform(OutMeshDescription, ComponentToWorldTransform);
-
-	// Use build settings from base mesh for LOD entries that were generated inside Editor.
-	const FMeshBuildSettings& BuildSettings = bImportedMesh ? StaticMeshModel.BuildSettings : StaticMesh->GetSourceModel(0).BuildSettings;
-
-	// Figure out if we should recompute normals and tangents. By default generated LODs should not recompute normals	
+	// Figure out if we should recompute normals and tangents. By default generated LODs should not recompute normals
 	EComputeNTBsFlags ComputeNTBsOptions = EComputeNTBsFlags::BlendOverlappingNormals;
 	if (BuildSettings.bRemoveDegenerates)
 	{
-		// If removing degenerate triangles, ignore them when computing tangents.
-		ComputeNTBsOptions |= EComputeNTBsFlags::IgnoreDegenerateTriangles;
-	}
-	if (BuildSettings.bUseMikkTSpace)
-	{
-		ComputeNTBsOptions |= EComputeNTBsFlags::UseMikkTSpace;
-	}
-
-	FStaticMeshOperations::ComputeTriangleTangentsAndNormals(OutMeshDescription, 0.0f);
-	FStaticMeshOperations::RecomputeNormalsAndTangentsIfNeeded(OutMeshDescription, ComputeNTBsOptions);
-}
-
-void FMeshMergeHelpers::RetrieveMesh(const USkeletalMeshComponent* SkeletalMeshComponent, int32 LODIndex, FMeshDescription& OutMeshDescription, bool bPropagateVertexColours)
-{
-	FSkeletalMeshModel* Resource = SkeletalMeshComponent->SkeletalMesh->GetImportedModel();
-	if (Resource->LODModels.IsValidIndex(LODIndex))
-	{
-		FSkeletalMeshLODInfo& SrcLODInfo = *(SkeletalMeshComponent->SkeletalMesh->GetLODInfo(LODIndex));
-
-		// Get the CPU skinned verts for this LOD
-		TArray<FFinalSkinVertex> FinalVertices;
-
-		// GetCPUSkinnedVertices 
-		SkeletalMeshComponent->GetCPUSkinnedVertices(FinalVertices, LODIndex);
-
-		FSkeletalMeshLODModel& LODModel = Resource->LODModels[LODIndex];
-		
-		const int32 NumSections = LODModel.Sections.Num();
-		
-		// Empty the mesh description
-		OutMeshDescription.Empty();
-
-		FStaticMeshAttributes Attributes(OutMeshDescription);
-		TVertexAttributesRef<FVector3f> VertexPositions = Attributes.GetVertexPositions();
-		TEdgeAttributesRef<bool> EdgeHardnesses = Attributes.GetEdgeHardnesses();
-		TPolygonGroupAttributesRef<FName> PolygonGroupImportedMaterialSlotNames = Attributes.GetPolygonGroupMaterialSlotNames();
-		TVertexInstanceAttributesRef<FVector3f> VertexInstanceNormals = Attributes.GetVertexInstanceNormals();
-		TVertexInstanceAttributesRef<FVector3f> VertexInstanceTangents = Attributes.GetVertexInstanceTangents();
-		TVertexInstanceAttributesRef<float> VertexInstanceBinormalSigns = Attributes.GetVertexInstanceBinormalSigns();
-		TVertexInstanceAttributesRef<FVector4f> VertexInstanceColors = Attributes.GetVertexInstanceColors();
-		TVertexInstanceAttributesRef<FVector2f> VertexInstanceUVs = Attributes.GetVertexInstanceUVs();
-
-		int32 TotalTriangles = 0;
-		int32 TotalCorners = 0;
-		for (int32 SectionIndex = 0; SectionIndex < NumSections; SectionIndex++)
-		{
-			const FSkelMeshSection& SkelMeshSection = LODModel.Sections[SectionIndex];
-			TotalTriangles += SkelMeshSection.NumTriangles;
-		}
-		TotalCorners = TotalTriangles * 3;
-		OutMeshDescription.ReserveNewVertices(FinalVertices.Num());
-		OutMeshDescription.ReserveNewPolygons(TotalTriangles);
-		OutMeshDescription.ReserveNewVertexInstances(TotalCorners);
-		OutMeshDescription.ReserveNewEdges(TotalCorners);
-
-		// Copy skinned vertex positions
-		for (int32 VertIndex = 0; VertIndex < FinalVertices.Num(); ++VertIndex)
-		{
-			const FVertexID VertexID = OutMeshDescription.CreateVertex();
-			VertexPositions[VertexID] = FinalVertices[VertIndex].Position;
-		}
-
-		VertexInstanceUVs.SetNumChannels(MAX_TEXCOORDS);
-
-		
-		for (int32 SectionIndex = 0; SectionIndex < NumSections; SectionIndex++)
-		{
-			const FSkelMeshSection& SkelMeshSection = LODModel.Sections[SectionIndex];
-			const int32 NumWedges = SkelMeshSection.NumTriangles * 3;
-
-			//Create the polygon group ID
-			int32 MaterialIndex = SkelMeshSection.MaterialIndex;
-			// use the remapping of material indices if there is a valid value
-			if (SrcLODInfo.LODMaterialMap.IsValidIndex(SectionIndex) && SrcLODInfo.LODMaterialMap[SectionIndex] != INDEX_NONE)
-			{
-				MaterialIndex = FMath::Clamp<int32>(SrcLODInfo.LODMaterialMap[SectionIndex], 0, SkeletalMeshComponent->SkeletalMesh->GetMaterials().Num() - 1);
-			}
-
-			FName ImportedMaterialSlotName = SkeletalMeshComponent->SkeletalMesh->GetMaterials()[MaterialIndex].ImportedMaterialSlotName;
-			const FPolygonGroupID SectionPolygonGroupID(SectionIndex);
-			if (!OutMeshDescription.IsPolygonGroupValid(SectionPolygonGroupID))
-			{
-				OutMeshDescription.CreatePolygonGroupWithID(SectionPolygonGroupID);
-				PolygonGroupImportedMaterialSlotNames[SectionPolygonGroupID] = ImportedMaterialSlotName;
-			}
-			int32 WedgeIndex = 0;
-			for (uint32 SectionTriangleIndex = 0; SectionTriangleIndex < SkelMeshSection.NumTriangles; ++SectionTriangleIndex)
-			{
-				FVertexID VertexIndexes[3];
-				TArray<FVertexInstanceID> VertexInstanceIDs;
-				VertexInstanceIDs.SetNum(3);
-				for (int32 CornerIndex = 0; CornerIndex < 3; ++CornerIndex, ++WedgeIndex)
-				{
-					const int32 VertexIndexForWedge = LODModel.IndexBuffer[SkelMeshSection.BaseIndex + WedgeIndex];
-					VertexIndexes[CornerIndex] = FVertexID(VertexIndexForWedge);
-					FVertexInstanceID VertexInstanceID = OutMeshDescription.CreateVertexInstance(VertexIndexes[CornerIndex]);
-					VertexInstanceIDs[CornerIndex] = VertexInstanceID;
-					
-					const FSoftSkinVertex& SoftVertex = SkelMeshSection.SoftVertices[VertexIndexForWedge - SkelMeshSection.BaseVertexIndex];
-					const FFinalSkinVertex& SkinnedVertex = FinalVertices[VertexIndexForWedge];
-
-					//Set NTBs
-					const FVector TangentX = SkinnedVertex.TangentX.ToFVector();
-					const FVector TangentZ = SkinnedVertex.TangentZ.ToFVector();
-					//@todo: do we need to inverse the sign between skeletalmesh and staticmesh, the old code was doing so.
-					const float TangentYSign = SkinnedVertex.TangentZ.ToFVector4f().W;
-					
-					VertexInstanceTangents[VertexInstanceID] = (FVector3f)TangentX;
-					VertexInstanceBinormalSigns[VertexInstanceID] = TangentYSign;
-					VertexInstanceNormals[VertexInstanceID] = (FVector3f)TangentZ;
-
-					for (uint32 TexCoordIndex = 0; TexCoordIndex < MAX_TEXCOORDS; TexCoordIndex++)
-					{
-						//Add this vertex instance tex coord
-						VertexInstanceUVs.Set(VertexInstanceID, TexCoordIndex, SoftVertex.UVs[TexCoordIndex]);
-					}
-
-					//Add this vertex instance color
-					VertexInstanceColors[VertexInstanceID] = bPropagateVertexColours ? FVector4f(FLinearColor(SoftVertex.Color)) : FVector4f(1.0f, 1.0f, 1.0f);
-				}
-				//Create a polygon from this triangle
-				const FPolygonID NewPolygonID = OutMeshDescription.CreatePolygon(SectionPolygonGroupID, VertexInstanceIDs);
-			}
-		}
-	}
-}
-
-void FMeshMergeHelpers::RetrieveMesh(const UStaticMesh* StaticMesh, int32 LODIndex, FMeshDescription& OutMeshDescription)
-{
-	TRACE_CPUPROFILER_EVENT_SCOPE(FMeshMergeHelpers::RetrieveMesh)
-
-	const FStaticMeshSourceModel& StaticMeshModel = StaticMesh->GetSourceModel(LODIndex);
-
-	// Imported meshes will have a valid mesh description
-	const bool bImportedMesh = StaticMesh->IsMeshDescriptionValid(LODIndex);
-
-	// Export the mesh data using static mesh render data
-	ExportStaticMeshLOD(StaticMesh->GetRenderData()->LODResources[LODIndex], OutMeshDescription, StaticMesh->GetStaticMaterials());
-
-	// Make sure the mesh is not irreparably malformed.
-	if (OutMeshDescription.VertexInstances().Num() <= 0)
-	{
-		return;
-	}
-
-	// Use build settings from base mesh for LOD entries that were generated inside Editor.
-	const FMeshBuildSettings& BuildSettings = bImportedMesh ? StaticMeshModel.BuildSettings : StaticMesh->GetSourceModel(0).BuildSettings;
-
-	// Figure out if we should recompute normals and tangents. By default generated LODs should not recompute normals	
-	EComputeNTBsFlags ComputeNTBsOptions = EComputeNTBsFlags::BlendOverlappingNormals;
-	if (BuildSettings.bRemoveDegenerates)
-	{
@@ -438,6 +297,16 @@
 	FStaticMeshOperations::RecomputeNormalsAndTangentsIfNeeded(OutMeshDescription, ComputeNTBsOptions);
 }
 
+void FMeshMergeHelpers::RetrieveMesh(const UStaticMeshComponent* StaticMeshComponent, int32 LODIndex, FMeshDescription& OutMeshDescription, bool bPropagateVertexColours)
+{
+	RetrieveMeshInternal(StaticMeshComponent->GetStaticMesh(), StaticMeshComponent, LODIndex, OutMeshDescription, bPropagateVertexColours);
+}
+
+void FMeshMergeHelpers::RetrieveMesh(const UStaticMesh* StaticMesh, int32 LODIndex, FMeshDescription& OutMeshDescription)
+{
+	RetrieveMeshInternal(StaticMesh, nullptr, LODIndex, OutMeshDescription, /*bPropagateVertexColours*/false);
+}
+
 void FMeshMergeHelpers::ExportStaticMeshLOD(const FStaticMeshLODResources& StaticMeshLOD, FMeshDescription& OutMeshDescription, const TArray<FStaticMaterial>& Materials)
 {
 	const int32 NumWedges = StaticMeshLOD.IndexBuffer.GetNumIndices();
@@ -551,6 +420,122 @@
 	}
 }
 
+void FMeshMergeHelpers::RetrieveMesh(const USkeletalMeshComponent* SkeletalMeshComponent, int32 LODIndex, FMeshDescription& OutMeshDescription, bool bPropagateVertexColours)
+{
+	FSkeletalMeshModel* Resource = SkeletalMeshComponent->SkeletalMesh->GetImportedModel();
+	if (Resource->LODModels.IsValidIndex(LODIndex))
+	{
+		FSkeletalMeshLODInfo& SrcLODInfo = *(SkeletalMeshComponent->SkeletalMesh->GetLODInfo(LODIndex));
+
+		// Get the CPU skinned verts for this LOD
+		TArray<FFinalSkinVertex> FinalVertices;
+
+		// GetCPUSkinnedVertices 
+		SkeletalMeshComponent->GetCPUSkinnedVertices(FinalVertices, LODIndex);
+
+		FSkeletalMeshLODModel& LODModel = Resource->LODModels[LODIndex];
+
+		const int32 NumSections = LODModel.Sections.Num();
+
+		// Empty the mesh description
+		OutMeshDescription.Empty();
+
+		FStaticMeshAttributes Attributes(OutMeshDescription);
+		TVertexAttributesRef<FVector3f> VertexPositions = Attributes.GetVertexPositions();
+		TEdgeAttributesRef<bool> EdgeHardnesses = Attributes.GetEdgeHardnesses();
+		TPolygonGroupAttributesRef<FName> PolygonGroupImportedMaterialSlotNames = Attributes.GetPolygonGroupMaterialSlotNames();
+		TVertexInstanceAttributesRef<FVector3f> VertexInstanceNormals = Attributes.GetVertexInstanceNormals();
+		TVertexInstanceAttributesRef<FVector3f> VertexInstanceTangents = Attributes.GetVertexInstanceTangents();
+		TVertexInstanceAttributesRef<float> VertexInstanceBinormalSigns = Attributes.GetVertexInstanceBinormalSigns();
+		TVertexInstanceAttributesRef<FVector4f> VertexInstanceColors = Attributes.GetVertexInstanceColors();
+		TVertexInstanceAttributesRef<FVector2f> VertexInstanceUVs = Attributes.GetVertexInstanceUVs();
+
+		int32 TotalTriangles = 0;
+		int32 TotalCorners = 0;
+		for (int32 SectionIndex = 0; SectionIndex < NumSections; SectionIndex++)
+		{
+			const FSkelMeshSection& SkelMeshSection = LODModel.Sections[SectionIndex];
+			TotalTriangles += SkelMeshSection.NumTriangles;
+		}
+		TotalCorners = TotalTriangles * 3;
+		OutMeshDescription.ReserveNewVertices(FinalVertices.Num());
+		OutMeshDescription.ReserveNewPolygons(TotalTriangles);
+		OutMeshDescription.ReserveNewVertexInstances(TotalCorners);
+		OutMeshDescription.ReserveNewEdges(TotalCorners);
+
+		// Copy skinned vertex positions
+		for (int32 VertIndex = 0; VertIndex < FinalVertices.Num(); ++VertIndex)
+		{
+			const FVertexID VertexID = OutMeshDescription.CreateVertex();
+			VertexPositions[VertexID] = FinalVertices[VertIndex].Position;
+		}
+
+		VertexInstanceUVs.SetNumChannels(MAX_TEXCOORDS);
+
+
+		for (int32 SectionIndex = 0; SectionIndex < NumSections; SectionIndex++)
+		{
+			const FSkelMeshSection& SkelMeshSection = LODModel.Sections[SectionIndex];
+			const int32 NumWedges = SkelMeshSection.NumTriangles * 3;
+
+			//Create the polygon group ID
+			int32 MaterialIndex = SkelMeshSection.MaterialIndex;
+			// use the remapping of material indices if there is a valid value
+			if (SrcLODInfo.LODMaterialMap.IsValidIndex(SectionIndex) && SrcLODInfo.LODMaterialMap[SectionIndex] != INDEX_NONE)
+			{
+				MaterialIndex = FMath::Clamp<int32>(SrcLODInfo.LODMaterialMap[SectionIndex], 0, SkeletalMeshComponent->SkeletalMesh->GetMaterials().Num() - 1);
+			}
+
+			FName ImportedMaterialSlotName = SkeletalMeshComponent->SkeletalMesh->GetMaterials()[MaterialIndex].ImportedMaterialSlotName;
+			const FPolygonGroupID SectionPolygonGroupID(SectionIndex);
+			if (!OutMeshDescription.IsPolygonGroupValid(SectionPolygonGroupID))
+			{
+				OutMeshDescription.CreatePolygonGroupWithID(SectionPolygonGroupID);
+				PolygonGroupImportedMaterialSlotNames[SectionPolygonGroupID] = ImportedMaterialSlotName;
+			}
+			int32 WedgeIndex = 0;
+			for (uint32 SectionTriangleIndex = 0; SectionTriangleIndex < SkelMeshSection.NumTriangles; ++SectionTriangleIndex)
+			{
+				FVertexID VertexIndexes[3];
+				TArray<FVertexInstanceID> VertexInstanceIDs;
+				VertexInstanceIDs.SetNum(3);
+				for (int32 CornerIndex = 0; CornerIndex < 3; ++CornerIndex, ++WedgeIndex)
+				{
+					const int32 VertexIndexForWedge = LODModel.IndexBuffer[SkelMeshSection.BaseIndex + WedgeIndex];
+					VertexIndexes[CornerIndex] = FVertexID(VertexIndexForWedge);
+					FVertexInstanceID VertexInstanceID = OutMeshDescription.CreateVertexInstance(VertexIndexes[CornerIndex]);
+					VertexInstanceIDs[CornerIndex] = VertexInstanceID;
+
+					const FSoftSkinVertex& SoftVertex = SkelMeshSection.SoftVertices[VertexIndexForWedge - SkelMeshSection.BaseVertexIndex];
+					const FFinalSkinVertex& SkinnedVertex = FinalVertices[VertexIndexForWedge];
+
+					//Set NTBs
+					const FVector TangentX = SkinnedVertex.TangentX.ToFVector();
+					const FVector TangentZ = SkinnedVertex.TangentZ.ToFVector();
+					//@todo: do we need to inverse the sign between skeletalmesh and staticmesh, the old code was doing so.
+					const float TangentYSign = SkinnedVertex.TangentZ.ToFVector4f().W;
+
+					VertexInstanceTangents[VertexInstanceID] = (FVector3f)TangentX;
+					VertexInstanceBinormalSigns[VertexInstanceID] = TangentYSign;
+					VertexInstanceNormals[VertexInstanceID] = (FVector3f)TangentZ;
+
+					for (uint32 TexCoordIndex = 0; TexCoordIndex < MAX_TEXCOORDS; TexCoordIndex++)
+					{
+						//Add this vertex instance tex coord
+						VertexInstanceUVs.Set(VertexInstanceID, TexCoordIndex, SoftVertex.UVs[TexCoordIndex]);
+					}
+
+					//Add this vertex instance color
+					VertexInstanceColors[VertexInstanceID] = bPropagateVertexColours ? FVector4f(FLinearColor(SoftVertex.Color)) : FVector4f(1.0f, 1.0f, 1.0f);
+				}
+				//Create a polygon from this triangle
+				const FPolygonID NewPolygonID = OutMeshDescription.CreatePolygon(SectionPolygonGroupID, VertexInstanceIDs);
+			}
+		}
+	}
+}
+
+
 bool FMeshMergeHelpers::CheckWrappingUVs(const TArray<FVector2D>& UVs)
 {	
 	bool bResult = false;
@@ -1028,17 +1013,16 @@
 
 bool FMeshMergeHelpers::PropagatePaintedColorsToMesh(const UStaticMeshComponent* StaticMeshComponent, int32 LODIndex, FMeshDescription& InOutMeshDescription)
 {
-	UStaticMesh* StaticMesh = StaticMeshComponent->GetStaticMesh();
-
-	if (StaticMesh->IsSourceModelValid(LODIndex) &&
-		StaticMeshComponent->LODData.IsValidIndex(LODIndex) &&
+	if (StaticMeshComponent->LODData.IsValidIndex(LODIndex) &&
 		StaticMeshComponent->LODData[LODIndex].OverrideVertexColors != nullptr)
 	{
 		FColorVertexBuffer& ColorVertexBuffer = *StaticMeshComponent->LODData[LODIndex].OverrideVertexColors;
-		FStaticMeshLODResources& RenderModel = StaticMesh->GetRenderData()->LODResources[LODIndex];
+		const UStaticMesh& StaticMesh = *StaticMeshComponent->GetStaticMesh();
 
-		if (ColorVertexBuffer.GetNumVertices() == RenderModel.GetNumVertices())
+		if (ColorVertexBuffer.GetNumVertices() == StaticMesh.GetNumVertices(LODIndex))
 		{	
+			const FStaticMeshLODResources& RenderModel = StaticMesh.GetRenderData()->LODResources[LODIndex];
+
 			const int32 NumWedges = InOutMeshDescription.VertexInstances().Num();
 			const int32 NumRenderWedges = RenderModel.IndexBuffer.GetNumIndices();
 			const bool bUseRenderWedges = NumWedges == NumRenderWedges;
diff --git a/Engine/Source/Runtime/MeshConversion/Public/StaticMeshLODResourcesAdapter.h b/Engine/Source/Runtime/MeshConversion/Public/StaticMeshLODResourcesAdapter.h
--- a/Engine/Source/Runtime/MeshConversion/Public/StaticMeshLODResourcesAdapter.h
+++ b/Engine/Source/Runtime/MeshConversion/Public/StaticMeshLODResourcesAdapter.h
@@ -18,10 +18,6 @@ struct FStaticMeshLODResourcesMeshAdapter
 protected:
 	const FStaticMeshLODResources* Mesh;
 
-	FVector3d BuildScale = FVector3d::One();
-	FVector3d InvBuildScale = FVector3d::One();
-	bool bScaleNormals = false;
-
 	TArray<const FStaticMeshSection*> ValidSections;
 	TArray<int32> TriangleOffsetArray;
 	int32 NumTriangles;
@@ -51,13 +47,6 @@ public:
 		TriangleOffsetArray.Add(NumTriangles);
 	}
 
-	void SetBuildScale(const FVector3d& BuildScaleIn, bool bScaleNormalsIn)
-	{
-		BuildScale = BuildScaleIn;
-		InvBuildScale = UE::Geometry::FTransformSRT3d::GetSafeScaleReciprocal(BuildScaleIn);
-		bScaleNormals = bScaleNormalsIn;
-	}
-
 	bool IsTriangle(int32 TID) const
 	{
 		return TID >= 0 && TID < MaxTriangleID();
@@ -109,7 +98,7 @@ public:
 	inline FVector3d GetVertex(int32 IDValue) const
 	{
 		const FVector3f& Position = Mesh->VertexBuffers.PositionVertexBuffer.VertexPosition(IDValue);
-		return FVector3d(BuildScale.X * (double)Position.X, BuildScale.Y * (double)Position.Y, BuildScale.Z * (double)Position.Z);
+		return FVector3d(Position.X, Position.Y, Position.Z);
 	}
 
 	inline void GetTriVertices(int32 IDValue, FVector3d& V0, FVector3d& V1, FVector3d& V2) const
@@ -151,8 +140,7 @@ public:
 	inline FVector3f GetNormal(int32 IDValue) const
 	{
 		const FVector4f& Normal = Mesh->VertexBuffers.StaticMeshVertexBuffer.VertexTangentZ(IDValue);
-		return (!bScaleNormals) ? FVector3f(Normal.X, Normal.Y, Normal.Z) :
-			UE::Geometry::Normalized(FVector3f(Normal.X * InvBuildScale.X, Normal.Y * InvBuildScale.Y, Normal.Z * InvBuildScale.Z));
+		return FVector3f(Normal.X, Normal.Y, Normal.Z);
 	}
 
 	/** Get Normals for a given Triangle */
@@ -170,8 +158,7 @@ public:
 	inline FVector3f GetTangentX(int32 IDValue) const
 	{
 		const FVector4f& TangentX = Mesh->VertexBuffers.StaticMeshVertexBuffer.VertexTangentX(IDValue);
-		return (!bScaleNormals) ? FVector3f(TangentX.X, TangentX.Y, TangentX.Z) :
-			UE::Geometry::Normalized(FVector3f(TangentX.X * BuildScale.X, TangentX.Y * BuildScale.Y, TangentX.Z * BuildScale.Z));
+		return FVector3f(TangentX.X, TangentX.Y, TangentX.Z);
 	}
 
 	/** Get Tangent X for a given Triangle */
@@ -188,8 +175,7 @@ public:
 	inline FVector3f GetTangentY(int32 IDValue) const
 	{
 		const FVector4f& TangentY = Mesh->VertexBuffers.StaticMeshVertexBuffer.VertexTangentY(IDValue);
-		return (!bScaleNormals) ? FVector3f(TangentY.X, TangentY.Y, TangentY.Z) :
-			UE::Geometry::Normalized(FVector3f(TangentY.X * BuildScale.X, TangentY.Y * BuildScale.Y, TangentY.Z * BuildScale.Z));
+		return FVector3f(TangentY.X, TangentY.Y, TangentY.Z);
 	}
 
 	/** Get Tangent Y for a given Triangle */
