Fix crash switching levels after building HLODs

HLOD meshes were sometime referencing MaterialInstanceDynamic, that aren't assets (!RF_Standalone)
* MIDs are usually outered to the components they are assigned to
* Duplicate MIDs and outer them to the generated HLOD UStaticMesh

diff --git a/Engine/Source/Developer/MeshMergeUtilities/Private/MeshMergeHelpers.cpp b/Engine/Source/Developer/MeshMergeUtilities/Private/MeshMergeHelpers.cpp
--- a/Engine/Source/Developer/MeshMergeUtilities/Private/MeshMergeHelpers.cpp
+++ b/Engine/Source/Developer/MeshMergeUtilities/Private/MeshMergeHelpers.cpp
@@ -1214,3 +1214,20 @@ void FMeshMergeHelpers::MergeImpostersToMesh(TArray<const UStaticMeshComponent*>
 		FStaticMeshOperations::AppendMeshDescription(ImposterMesh, InMeshDescription, AppendSettings);
 	}
 }
+
+void FMeshMergeHelpers::FixupNonStandaloneMaterialReferences(UStaticMesh* InStaticMesh)
+{
+	UPackage* Package = InStaticMesh->GetPackage();
+
+	// Ensure mesh is not referencing non standalone materials
+	// This can be the case for MaterialInstanceDynamic (MID) as they are normally outered to the components
+	for (FStaticMaterial& StaticMaterial : InStaticMesh->GetStaticMaterials())
+	{
+		TObjectPtr<UMaterialInterface>& MaterialInterface = StaticMaterial.MaterialInterface;
+		if (!MaterialInterface->HasAnyFlags(RF_Standalone) && MaterialInterface->GetPackage() != Package)
+		{
+			// Duplicate material and outer it to the mesh
+			MaterialInterface = DuplicateObject<UMaterialInterface>(MaterialInterface, InStaticMesh);
+		}
+	}
+}
diff --git a/Engine/Source/Developer/MeshMergeUtilities/Private/MeshMergeHelpers.h b/Engine/Source/Developer/MeshMergeUtilities/Private/MeshMergeHelpers.h
--- a/Engine/Source/Developer/MeshMergeUtilities/Private/MeshMergeHelpers.h
+++ b/Engine/Source/Developer/MeshMergeUtilities/Private/MeshMergeHelpers.h
@@ -74,4 +74,7 @@ public:
 	
 	/** Merges imposter meshes into the given MeshDescription. */
 	static void MergeImpostersToMesh(TArray<const UStaticMeshComponent*> ImposterComponents, FMeshDescription& InOutMeshDescription, const FVector& InPivot, int32 BaseMaterialIndex, TArray<UMaterialInterface*>& OutImposterMaterials);
+
+	/** Ensure a generated HLOD mesh is not referencing non standalone materials. */
+	static void FixupNonStandaloneMaterialReferences(UStaticMesh* InStaticMesh);
 };
diff --git a/Engine/Source/Developer/MeshMergeUtilities/Private/MeshMergeUtilities.cpp b/Engine/Source/Developer/MeshMergeUtilities/Private/MeshMergeUtilities.cpp
--- a/Engine/Source/Developer/MeshMergeUtilities/Private/MeshMergeUtilities.cpp
+++ b/Engine/Source/Developer/MeshMergeUtilities/Private/MeshMergeUtilities.cpp
@@ -3055,6 +3055,10 @@ void FMeshMergeUtilities::MergeComponentsToStaticMesh(const TArray<UPrimitiveCom
 			StaticMesh->BuildCacheAutomationTestGuid = FGuid::NewGuid();
 		}
 #endif
+
+		// Ensure the new mesh is not referencing non standalone materials
+		FMeshMergeHelpers::FixupNonStandaloneMaterialReferences(StaticMesh);
+
 		StaticMesh->Build(bSilent);
 
 		if (ImposterBounds.IsValid)
diff --git a/Engine/Source/Developer/MeshMergeUtilities/Private/ProxyGenerationProcessor.cpp b/Engine/Source/Developer/MeshMergeUtilities/Private/ProxyGenerationProcessor.cpp
--- a/Engine/Source/Developer/MeshMergeUtilities/Private/ProxyGenerationProcessor.cpp
+++ b/Engine/Source/Developer/MeshMergeUtilities/Private/ProxyGenerationProcessor.cpp
@@ -286,6 +286,9 @@ void FProxyGenerationProcessor::ProcessJob(const FGuid& JobGuid, FProxyGeneratio
 			}
 			StaticMesh->GetStaticMaterials().Add(NewMaterial);
 		}
+
+		// Ensure the new mesh is not referencing non standalone materials
+		FMeshMergeHelpers::FixupNonStandaloneMaterialReferences(StaticMesh);
 	}
 	else
 	{
