MeshMergeUtilities - Fix some logic that was locking the main thread while waiting for the StaticMesh async build to complete. Now the build is only triggered (by PostEditChange()) after we've done modifying the StaticMesh.

On a large test asset, this shaved off ~20 seconds on a ~1m50sec merge operation.

diff --git a/Engine/Source/Developer/MeshMergeUtilities/Private/MeshMergeUtilities.cpp b/Engine/Source/Developer/MeshMergeUtilities/Private/MeshMergeUtilities.cpp
--- a/Engine/Source/Developer/MeshMergeUtilities/Private/MeshMergeUtilities.cpp
+++ b/Engine/Source/Developer/MeshMergeUtilities/Private/MeshMergeUtilities.cpp
@@ -3059,7 +3059,6 @@ void FMeshMergeUtilities::MergeComponentsToStaticMesh(const TArray<UPrimitiveCom
 		// Ensure the new mesh is not referencing non standalone materials
 		FMeshMergeHelpers::FixupNonStandaloneMaterialReferences(StaticMesh);
 
-		StaticMesh->Build(bSilent);
 
 		if (ImposterBounds.IsValid)
 		{
@@ -3070,8 +3069,6 @@ void FMeshMergeUtilities::MergeComponentsToStaticMesh(const TArray<UPrimitiveCom
 			StaticMesh->CalculateExtendedBounds();
 		}		
 
-		StaticMesh->PostEditChange();
-
 		if (InSettings.bCreateMergedMaterial && MergedMaterial)
 		{
 			//Make sure we have unique slot name here
@@ -3082,9 +3079,10 @@ void FMeshMergeUtilities::MergeComponentsToStaticMesh(const TArray<UPrimitiveCom
 				MaterialSlotName = *(MergedMaterial->GetName() + TEXT("_") + FString::FromInt(Counter++));
 			}
 			StaticMesh->GetStaticMaterials().Add(FStaticMaterial(MergedMaterial, MaterialSlotName));
-			StaticMesh->UpdateUVChannelData(false);
 		}
 
+		StaticMesh->PostEditChange();
+
 		OutAssetsToSync.Add(StaticMesh);
 		OutMergedActorLocation = MergedAssetPivot;
 	}
