Detect failure of global UV remap and avoid a crash accessing the empty GlobalTextureCoordinates array

diff --git a/Engine/Source/Developer/MeshMergeUtilities/Private/MeshMergeUtilities.cpp b/Engine/Source/Developer/MeshMergeUtilities/Private/MeshMergeUtilities.cpp
--- a/Engine/Source/Developer/MeshMergeUtilities/Private/MeshMergeUtilities.cpp
+++ b/Engine/Source/Developer/MeshMergeUtilities/Private/MeshMergeUtilities.cpp
@@ -2560,39 +2560,46 @@ void FMeshMergeUtilities::MergeComponentsToStaticMesh(const TArray<UPrimitiveCom
 
 			// Create texture coords for the merged mesh
 			TArray<FVector2D> GlobalTextureCoordinates;
-			FStaticMeshOperations::GenerateUniqueUVsForStaticMesh(MergedRawMeshes[0], MaterialOptions->TextureSize.GetMax(), true, GlobalTextureCoordinates);
-			ScaleTextureCoordinatesToBox(FBox2D(FVector2D::ZeroVector, FVector2D(1, 1)), GlobalTextureCoordinates);
-
-			// copy UVs back to the un-merged mesh's custom texture coords
-			// iterate the raw meshes in the same way as when we combined the mesh above in CreateMergedRawMeshes()
-			int32 GlobalUVIndex = 0;
-			for (TConstRawMeshIterator RawMeshIterator = DataTracker.GetConstRawMeshIterator(); RawMeshIterator; ++RawMeshIterator)
+			bool bSuccess = FStaticMeshOperations::GenerateUniqueUVsForStaticMesh(MergedRawMeshes[0], MaterialOptions->TextureSize.GetMax(), true, GlobalTextureCoordinates);
+			if (bSuccess)
 			{
-				const FMeshLODKey& Key = RawMeshIterator.Key();
-				const FMeshDescription& RawMesh = RawMeshIterator.Value();
+				ScaleTextureCoordinatesToBox(FBox2D(FVector2D::ZeroVector, FVector2D(1, 1)), GlobalTextureCoordinates);
 
-				// Build a local array for this raw mesh
-				TArray<FVector2D> UniqueTextureCoordinates;
-				UniqueTextureCoordinates.SetNumUninitialized(RawMesh.VertexInstances().Num());
-				for(FVector2D& UniqueTextureCoordinate : UniqueTextureCoordinates)
+				// copy UVs back to the un-merged mesh's custom texture coords
+				// iterate the raw meshes in the same way as when we combined the mesh above in CreateMergedRawMeshes()
+				int32 GlobalUVIndex = 0;
+				for (TConstRawMeshIterator RawMeshIterator = DataTracker.GetConstRawMeshIterator(); RawMeshIterator; ++RawMeshIterator)
 				{
-					UniqueTextureCoordinate = GlobalTextureCoordinates[GlobalUVIndex++];
-				}
+					const FMeshLODKey& Key = RawMeshIterator.Key();
+					const FMeshDescription& RawMesh = RawMeshIterator.Value();
 
-				// copy to mesh data
-				for(FMeshData& MeshData : GlobalMeshSettings)
-				{
-					if(MeshData.MeshDescription == &RawMesh)
+					// Build a local array for this raw mesh
+					TArray<FVector2D> UniqueTextureCoordinates;
+					UniqueTextureCoordinates.SetNumUninitialized(RawMesh.VertexInstances().Num());
+					for (FVector2D& UniqueTextureCoordinate : UniqueTextureCoordinates)
 					{
-						MeshData.CustomTextureCoordinates = UniqueTextureCoordinates;
+						UniqueTextureCoordinate = GlobalTextureCoordinates[GlobalUVIndex++];
+					}
+
+					// copy to mesh data
+					for (FMeshData& MeshData : GlobalMeshSettings)
+					{
+						if (MeshData.MeshDescription == &RawMesh)
+						{
+							MeshData.CustomTextureCoordinates = UniqueTextureCoordinates;
+						}
 					}
 				}
-			}
 
-			// Dont smear borders as we will copy back non-pink pixels
-			for(FMaterialData& MaterialData : GlobalMaterialSettings)
+				// Dont smear borders as we will copy back non-pink pixels
+				for (FMaterialData& MaterialData : GlobalMaterialSettings)
+				{
+					MaterialData.bPerformBorderSmear = false;
+				}
+			}
+			else
 			{
-				MaterialData.bPerformBorderSmear = false;
+				UE_LOG(LogMeshMerging, Error, TEXT("GenerateUniqueUVsForStaticMesh: Failed to pack UVs for static mesh"));
 			}
 		}
 
