Merge Actor: Fixed incorrect MeshDataIndex stored as the key to the MeshDataAwaitingResults map

diff --git a/Engine/Source/Developer/MeshMergeUtilities/Private/MeshMergeUtilities.cpp b/Engine/Source/Developer/MeshMergeUtilities/Private/MeshMergeUtilities.cpp
--- a/Engine/Source/Developer/MeshMergeUtilities/Private/MeshMergeUtilities.cpp
+++ b/Engine/Source/Developer/MeshMergeUtilities/Private/MeshMergeUtilities.cpp
@@ -2352,6 +2352,7 @@ void FMeshMergeUtilities::MergeComponentsToStaticMesh(const TArray<UPrimitiveCom
 		}
 
 		UMaterialOptions* MaterialOptions = PopulateMaterialOptions(MaterialProxySettings);
+
 		// Check each material to see if the shader actually uses vertex data and collect flags
 		TArray<bool> bMaterialUsesVertexData;
 		DetermineMaterialVertexDataUsage(bMaterialUsesVertexData, UniqueMaterials, MaterialOptions);
@@ -2369,8 +2370,6 @@ void FMeshMergeUtilities::MergeComponentsToStaticMesh(const TArray<UPrimitiveCom
 			}
 		}
 
-		TMap<UMaterialInterface*, int32> MaterialToDefaultMeshData;
-
 		// If we are generating a single LOD and want to merge materials we can utilize texture space better by generating unique UVs
 		// for the merged mesh and baking out materials using those UVs
 		const bool bGloballyRemapUVs = !bMergeAllLODs && !InSettings.bReuseMeshLightmapUVs;
@@ -2385,7 +2384,10 @@ void FMeshMergeUtilities::MergeComponentsToStaticMesh(const TArray<UPrimitiveCom
 			const FMeshLODKey& Key = RawMeshIterator.Key();
 			const FMeshDescription& RawMesh = RawMeshIterator.Value();
 			const bool bRequiresUniqueUVs = DataTracker.DoesMeshLODRequireUniqueUVs(Key);
+
+			FMeshDescription* MeshDescription = DataTracker.GetRawMeshPtr(Key);
 			UStaticMeshComponent* Component = StaticMeshComponentsToMerge[Key.GetMeshIndex()];
+			UStaticMesh* StaticMesh = Component->GetStaticMesh();
 
 			// Retrieve all sections and materials for key
 			TArray<SectionRemapPair> SectionRemapPairs;
@@ -2413,23 +2415,59 @@ void FMeshMergeUtilities::MergeComponentsToStaticMesh(const TArray<UPrimitiveCom
 				MaterialData.Material = CollapsedMaterialMap.FindChecked(Material);
 				MaterialData.PropertySizes = PropertySizes;
 
-				FMeshData MeshData;
-				MeshData.Mesh = Key.GetMesh();
-				MeshData.VertexColorHash = Key.GetVertexColorHash();
-				MeshData.bMirrored = Component->GetComponentTransform().GetDeterminant() < 0.0f;
-				int32 MeshDataIndex = 0;
-				
-				if (InSettings.bCreateMergedMaterial || bGloballyRemapUVs || (InSettings.bUseVertexDataForBakingMaterial && (bDoesMaterialUseVertexData || bRequiresUniqueUVs)))
+				FMeshData NewMeshData;
+				const bool bUseMeshData = InSettings.bCreateMergedMaterial || bGloballyRemapUVs || (InSettings.bUseVertexDataForBakingMaterial && (bDoesMaterialUseVertexData || bRequiresUniqueUVs));
+				if (bUseMeshData)
 				{
-					FMeshDescription* MeshDescription = DataTracker.GetRawMeshPtr(Key);
-					MeshData.MeshDescription = MeshDescription;
+					NewMeshData.Mesh = Key.GetMesh();
+					NewMeshData.MeshDescription = MeshDescription;
+					NewMeshData.VertexColorHash = Key.GetVertexColorHash();
+					NewMeshData.bMirrored = Component->GetComponentTransform().GetDeterminant() < 0.0f;					
+					NewMeshData.MaterialIndices = SectionIndices;
+				}
 
+				auto CompareMaterialData = [&InSettings](const FMaterialData& LHS, const FMaterialData& RHS)
+				{
+					return InSettings.bMergeEquivalentMaterials ? FMaterialKey(LHS.Material) == FMaterialKey(RHS.Material) : LHS.Material == RHS.Material;
+				};
+
+				auto CompareMeshData = [](const FMeshData& LHS, const FMeshData& RHS)
+				{
+					return (LHS.Mesh == RHS.Mesh) && (LHS.MaterialIndices == RHS.MaterialIndices) && (LHS.bMirrored == RHS.bMirrored) && (LHS.VertexColorHash == RHS.VertexColorHash);
+				};
+
+				// Find material & mesh pair
+				bool FoundMatch = false;
+				for (int32 GlobalMaterialSettingsIndex = 0; GlobalMaterialSettingsIndex < GlobalMaterialSettings.Num(); ++GlobalMaterialSettingsIndex)
+				{
+					if (CompareMaterialData(GlobalMaterialSettings[GlobalMaterialSettingsIndex], MaterialData) &&
+						CompareMeshData(GlobalMeshSettings[GlobalMaterialSettingsIndex], NewMeshData))
+					{
+						FoundMatch = true;
+						break;
+					}
+				}
+
+				// We've found a match, no need to process this mesh/material pair
+				if (FoundMatch)
+				{
+					continue;
+				}
+
+				// We're processing a new pair
+				int32 MeshDataIndex = GlobalMeshSettings.Num();
+				FMeshData& MeshData = GlobalMeshSettings.Emplace_GetRef(NewMeshData);
+				GlobalMaterialSettings.Add(MaterialData);
+				SectionMaterialImportanceValues.Add(MaterialImportanceValues[MaterialIndex]);
+				
+				if (bUseMeshData)
+				{
 					// if it has vertex color/*WedgetColors.Num()*/, it should also use light map UV index
 					// we can't do this for all meshes, but only for the mesh that has vertex color.
 					if (bRequiresUniqueUVs || MeshData.MeshDescription->VertexInstances().Num() > 0)
 					{
 						// Check if there are lightmap uvs available?
-						const int32 LightMapUVIndex = StaticMeshComponentsToMerge[Key.GetMeshIndex()]->GetStaticMesh()->GetLightMapCoordinateIndex();
+						const int32 LightMapUVIndex = StaticMesh->GetLightMapCoordinateIndex();
 
 						TVertexInstanceAttributesConstRef<FVector2f> VertexInstanceUVs = FStaticMeshConstAttributes(*MeshData.MeshDescription).GetVertexInstanceUVs();
 						if (InSettings.bReuseMeshLightmapUVs && VertexInstanceUVs.GetNumElements() > 0 && VertexInstanceUVs.GetNumChannels() > LightMapUVIndex)
@@ -2456,72 +2494,7 @@ void FMeshMergeUtilities::MergeComponentsToStaticMesh(const TArray<UPrimitiveCom
 						}
 					}
 
-					MeshData.TextureCoordinateBox = FBox2D(FVector2D(0.0f, 0.0f), FVector2D(1.0f, 1.0f));
-					MeshData.MaterialIndices = SectionIndices;
-					MeshDataIndex = GlobalMeshSettings.Num();
-
 					Adapters[Key.GetMeshIndex()].ApplySettings(Key.GetLODIndex(), MeshData);
-
-					int32 ExistingMeshDataIndex = INDEX_NONE;
-
-					auto MaterialsAreEquivalent = [&InSettings](const UMaterialInterface* Material0, const UMaterialInterface* Material1)
-					{
-						if (InSettings.bMergeEquivalentMaterials)
-						{
-							return FMaterialKey(Material0) == FMaterialKey(Material1);
-						}
-						else
-						{
-							return Material0 == Material1;
-						}
-					};
-
-					// Find any existing materials
-					for (int32 GlobalMaterialSettingsIndex = 0; GlobalMaterialSettingsIndex < GlobalMaterialSettings.Num(); ++GlobalMaterialSettingsIndex)
-					{
-						const FMaterialData& ExistingMaterialData = GlobalMaterialSettings[GlobalMaterialSettingsIndex];
-						// Compare materials (note this assumes property sizes match!)
-						if (MaterialsAreEquivalent(ExistingMaterialData.Material, MaterialData.Material))
-						{
-							// materials match, so check the corresponding mesh data
-							const FMeshData& ExistingMeshData = GlobalMeshSettings[GlobalMaterialSettingsIndex];
-							bool bMatchesMesh = (ExistingMeshData.Mesh == MeshData.Mesh &&
-								ExistingMeshData.MaterialIndices == MeshData.MaterialIndices &&
-								ExistingMeshData.bMirrored == MeshData.bMirrored &&
-								ExistingMeshData.VertexColorHash == MeshData.VertexColorHash);
-							if (bMatchesMesh)
-							{
-								MeshDataIndex = ExistingMeshDataIndex = GlobalMaterialSettingsIndex;
-								break;
-							}
-						}
-					}
-
-					if (ExistingMeshDataIndex == INDEX_NONE)
-					{
-						GlobalMeshSettings.Add(MeshData);
-						GlobalMaterialSettings.Add(MaterialData);
-						SectionMaterialImportanceValues.Add(MaterialImportanceValues[MaterialIndex]);
-					}
-				}
-				else
-				{
-					MeshData.MeshDescription = nullptr;
-					MeshData.TextureCoordinateBox = FBox2D(FVector2D(0.0f, 0.0f), FVector2D(1.0f, 1.0f));
-
-					// This prevents baking out the same material multiple times, which would be wasteful when it does not use vertex data anyway
-					const bool bPreviouslyAdded = MaterialToDefaultMeshData.Contains(Material);
-					int32& DefaultMeshDataIndex = MaterialToDefaultMeshData.FindOrAdd(Material);
-
-					if (!bPreviouslyAdded)
-					{
-						DefaultMeshDataIndex = GlobalMeshSettings.Num();
-						GlobalMeshSettings.Add(MeshData);
-						GlobalMaterialSettings.Add(MaterialData);
-						SectionMaterialImportanceValues.Add(MaterialImportanceValues[MaterialIndex]);
-					}
-
-					MeshDataIndex = DefaultMeshDataIndex;
 				}
 
 				for (const auto& OriginalSectionIndex : SectionIndices)
