Merge Actor - Fix section/material mapping issue introduced by CL22158948

diff --git a/Engine/Source/Developer/MeshMergeUtilities/Private/MeshMergeUtilities.cpp b/Engine/Source/Developer/MeshMergeUtilities/Private/MeshMergeUtilities.cpp
--- a/Engine/Source/Developer/MeshMergeUtilities/Private/MeshMergeUtilities.cpp
+++ b/Engine/Source/Developer/MeshMergeUtilities/Private/MeshMergeUtilities.cpp
@@ -2437,64 +2437,63 @@ void FMeshMergeUtilities::MergeComponentsToStaticMesh(const TArray<UPrimitiveCom
 				};
 
 				// Find material & mesh pair
-				bool FoundMatch = false;
+				int32 MeshDataIndex = INDEX_NONE;
 				for (int32 GlobalMaterialSettingsIndex = 0; GlobalMaterialSettingsIndex < GlobalMaterialSettings.Num(); ++GlobalMaterialSettingsIndex)
 				{
 					if (CompareMaterialData(GlobalMaterialSettings[GlobalMaterialSettingsIndex], MaterialData) &&
 						CompareMeshData(GlobalMeshSettings[GlobalMaterialSettingsIndex], NewMeshData))
 					{
-						FoundMatch = true;
+						MeshDataIndex = GlobalMaterialSettingsIndex;
 						break;
 					}
 				}
 
 				// We've found a match, no need to process this mesh/material pair
-				if (FoundMatch)
+				if (MeshDataIndex == INDEX_NONE)
 				{
-					continue;
-				}
+					// We're processing a new pair
+					MeshDataIndex = GlobalMeshSettings.Num();
 
-				// We're processing a new pair
-				int32 MeshDataIndex = GlobalMeshSettings.Num();
-				FMeshData& MeshData = GlobalMeshSettings.Emplace_GetRef(NewMeshData);
-				GlobalMaterialSettings.Add(MaterialData);
-				SectionMaterialImportanceValues.Add(MaterialImportanceValues[MaterialIndex]);
+					FMeshData& MeshData = GlobalMeshSettings.Emplace_GetRef(NewMeshData);
+					GlobalMaterialSettings.Add(MaterialData);
+					SectionMaterialImportanceValues.Add(MaterialImportanceValues[MaterialIndex]);
 				
-				if (bUseMeshData)
-				{
-					// if it has vertex color/*WedgetColors.Num()*/, it should also use light map UV index
-					// we can't do this for all meshes, but only for the mesh that has vertex color.
-					if (bRequiresUniqueUVs || MeshData.MeshDescription->VertexInstances().Num() > 0)
+					if (bUseMeshData)
 					{
-						// Check if there are lightmap uvs available?
-						const int32 LightMapUVIndex = StaticMesh->GetLightMapCoordinateIndex();
-
-						TVertexInstanceAttributesConstRef<FVector2f> VertexInstanceUVs = FStaticMeshConstAttributes(*MeshData.MeshDescription).GetVertexInstanceUVs();
-						if (InSettings.bReuseMeshLightmapUVs && VertexInstanceUVs.GetNumElements() > 0 && VertexInstanceUVs.GetNumChannels() > LightMapUVIndex)
+						// if it has vertex color/*WedgetColors.Num()*/, it should also use light map UV index
+						// we can't do this for all meshes, but only for the mesh that has vertex color.
+						if (bRequiresUniqueUVs || MeshData.MeshDescription->VertexInstances().Num() > 0)
 						{
-							MeshData.TextureCoordinateIndex = LightMapUVIndex;
-						}
-						else
-						{
-							// Verify if we started an async task to generate UVs for this static mesh & LOD
-							FMeshLODTuple Tuple(Key.GetMesh(), Key.GetLODIndex());
-							if (!MeshLODsTextureCoordinates.Find(Tuple))
+							// Check if there are lightmap uvs available?
+							const int32 LightMapUVIndex = StaticMesh->GetLightMapCoordinateIndex();
+
+							TVertexInstanceAttributesConstRef<FVector2f> VertexInstanceUVs = FStaticMeshConstAttributes(*MeshData.MeshDescription).GetVertexInstanceUVs();
+							if (InSettings.bReuseMeshLightmapUVs && VertexInstanceUVs.GetNumElements() > 0 && VertexInstanceUVs.GetNumChannels() > LightMapUVIndex)
 							{
-								// No job found yet, fire an async task
-								MeshLODsTextureCoordinates.Add(Tuple, Async(EAsyncExecution::Thread, [MeshDescription, MaterialOptions, this]()
+								MeshData.TextureCoordinateIndex = LightMapUVIndex;
+							}
+							else
+							{
+								// Verify if we started an async task to generate UVs for this static mesh & LOD
+								FMeshLODTuple Tuple(Key.GetMesh(), Key.GetLODIndex());
+								if (!MeshLODsTextureCoordinates.Find(Tuple))
 								{
-									TArray<FVector2D> UniqueTextureCoordinates;
-									FStaticMeshOperations::GenerateUniqueUVsForStaticMesh(*MeshDescription, MaterialOptions->TextureSize.GetMax(), false, UniqueTextureCoordinates);
-									ScaleTextureCoordinatesToBox(FBox2D(FVector2D::ZeroVector, FVector2D(1, 1)), UniqueTextureCoordinates);
-									return UniqueTextureCoordinates;
-								}));
+									// No job found yet, fire an async task
+									MeshLODsTextureCoordinates.Add(Tuple, Async(EAsyncExecution::Thread, [MeshDescription, MaterialOptions, this]()
+									{
+										TArray<FVector2D> UniqueTextureCoordinates;
+										FStaticMeshOperations::GenerateUniqueUVsForStaticMesh(*MeshDescription, MaterialOptions->TextureSize.GetMax(), false, UniqueTextureCoordinates);
+										ScaleTextureCoordinatesToBox(FBox2D(FVector2D::ZeroVector, FVector2D(1, 1)), UniqueTextureCoordinates);
+										return UniqueTextureCoordinates;
+									}));
+								}
+								// Keep track of the fact that this mesh is waiting for the UV computation to finish
+								MeshDataAwaitingResults.Add(MeshDataIndex, Tuple);
 							}
-							// Keep track of the fact that this mesh is waiting for the UV computation to finish
-							MeshDataAwaitingResults.Add(MeshDataIndex, Tuple);
 						}
-					}
 
-					Adapters[Key.GetMeshIndex()].ApplySettings(Key.GetLODIndex(), MeshData);
+						Adapters[Key.GetMeshIndex()].ApplySettings(Key.GetLODIndex(), MeshData);
+					}
 				}
 
 				for (const auto& OriginalSectionIndex : SectionIndices)
