Merge Actor: Turn a UV generation error into a warning as a temp workaround for failing HLOD builds
* Will rework UV generation to be more resilient

diff --git a/Engine/Source/Developer/MeshMergeUtilities/Private/MeshMergeUtilities.cpp b/Engine/Source/Developer/MeshMergeUtilities/Private/MeshMergeUtilities.cpp
--- a/Engine/Source/Developer/MeshMergeUtilities/Private/MeshMergeUtilities.cpp
+++ b/Engine/Source/Developer/MeshMergeUtilities/Private/MeshMergeUtilities.cpp
@@ -2571,7 +2571,7 @@ void FMeshMergeUtilities::MergeComponentsToStaticMesh(const TArray<UPrimitiveCom
 			}
 			else
 			{
-				UE_LOG(LogMeshMerging, Error, TEXT("GenerateUniqueUVsForStaticMesh: Failed to pack UVs for static mesh"));
+				UE_LOG(LogMeshMerging, Warning, TEXT("GenerateUniqueUVsForStaticMesh: Failed to pack UVs for static mesh"));
 			}
 		}
 
