World Partition - HLOD: Added null material check in FMeshMergeHelpers::FixupNonStandaloneMaterialReferences()

diff --git a/Engine/Source/Developer/MeshMergeUtilities/Private/MeshMergeHelpers.cpp b/Engine/Source/Developer/MeshMergeUtilities/Private/MeshMergeHelpers.cpp
--- a/Engine/Source/Developer/MeshMergeUtilities/Private/MeshMergeHelpers.cpp
+++ b/Engine/Source/Developer/MeshMergeUtilities/Private/MeshMergeHelpers.cpp
@@ -1228,7 +1228,7 @@ void FMeshMergeHelpers::FixupNonStandaloneMaterialReferences(UStaticMesh* InStat
 	for (FStaticMaterial& StaticMaterial : InStaticMesh->GetStaticMaterials())
 	{
 		TObjectPtr<UMaterialInterface>& MaterialInterface = StaticMaterial.MaterialInterface;
-		if (!MaterialInterface->HasAnyFlags(RF_Standalone) && MaterialInterface->GetPackage() != Package)
+		if (MaterialInterface && !MaterialInterface->HasAnyFlags(RF_Standalone) && MaterialInterface->GetPackage() != Package)
 		{
 			// Duplicate material and outer it to the mesh
 			MaterialInterface = DuplicateObject<UMaterialInterface>(MaterialInterface, InStaticMesh);
