Merge Actors: don't dirty the level when merging actors.

diff --git a/Engine/Source/Developer/MeshMergeUtilities/Private/MeshMergeUtilities.cpp b/Engine/Source/Developer/MeshMergeUtilities/Private/MeshMergeUtilities.cpp
--- a/Engine/Source/Developer/MeshMergeUtilities/Private/MeshMergeUtilities.cpp
+++ b/Engine/Source/Developer/MeshMergeUtilities/Private/MeshMergeUtilities.cpp
@@ -3494,7 +3494,7 @@ void FMeshMergeUtilities::MergeComponentsToInstances(const TArray<UPrimitiveComp
 			{
 				// Create our actors
 				const FScopedTransaction Transaction(LOCTEXT("PlaceInstancedActors", "Place Instanced Actor(s)"));
-				Level->Modify();
+				Level->Modify(false);
 
  				FActorSpawnParameters Params;
  				Params.OverrideLevel = Level;
