Merging static mesh actors resulted in static mesh sockets getting lost
* Added a "Merge Mesh Sockets" option to the Mesh Merge Settings

diff --git a/Engine/Source/Developer/MeshMergeUtilities/Private/MeshMergeUtilities.cpp b/Engine/Source/Developer/MeshMergeUtilities/Private/MeshMergeUtilities.cpp
--- a/Engine/Source/Developer/MeshMergeUtilities/Private/MeshMergeUtilities.cpp
+++ b/Engine/Source/Developer/MeshMergeUtilities/Private/MeshMergeUtilities.cpp
@@ -4,6 +4,7 @@
 
 #include "Engine/MapBuildDataRegistry.h"
 #include "Engine/MeshMerging.h"
+#include "Engine/StaticMeshSocket.h"
 
 #include "MaterialOptions.h"
 #include "IMaterialBakingModule.h"
@@ -2302,6 +2303,46 @@ void FMeshMergeUtilities::MergeComponentsToStaticMesh(const TArray<UPrimitiveCom
 		ExtractPhysicsDataFromComponents(ComponentsToMerge, PhysicsGeometry, BodySetupSource);
 	}
 
+	// Merge sockets
+	TMap<FName, UStaticMeshSocket*> MergedSockets;
+	if (InSettings.bMergeMeshSockets)
+	{
+		const FTransform PivotTransform = FTransform(MergedAssetPivot);
+		for (UPrimitiveComponent* PrimitiveComponent : ComponentsToMerge)
+		{
+			if (UStaticMeshComponent* StaticMeshComponent = Cast<UStaticMeshComponent>(PrimitiveComponent))
+			{
+				if (UStaticMesh* StaticMesh = StaticMeshComponent->GetStaticMesh())
+				{
+					for (UStaticMeshSocket* Socket : StaticMesh->Sockets)
+					{
+						if (Socket)
+						{
+							UStaticMeshSocket* SocketCopy = DuplicateObject<UStaticMeshSocket>(Socket, nullptr);
+
+						    // Fix name - rename if duplicates are found
+							FString PlainName = SocketCopy->SocketName.GetPlainNameString();
+						    int32 CurrentNumber = SocketCopy->SocketName.GetNumber();
+						    while (MergedSockets.Contains(SocketCopy->SocketName))
+						    {
+							    SocketCopy->SocketName = FName(PlainName, CurrentNumber++);
+						    }
+    
+						    // Fix transform - make relative to pivot
+						    FTransform SocketTransformWorldSpace = StaticMeshComponent->GetSocketTransform(Socket->SocketName, RTS_World);
+						    FTransform SocketTransformPivotSpace = SocketTransformWorldSpace.GetRelativeTransform(PivotTransform);
+						    SocketCopy->RelativeLocation = SocketTransformPivotSpace.GetLocation();
+						    SocketCopy->RelativeRotation = FRotator(SocketTransformPivotSpace.GetRotation());
+						    SocketCopy->RelativeScale = SocketTransformPivotSpace.GetScale3D();
+
+							MergedSockets.Add(SocketCopy->SocketName, SocketCopy);
+						}
+					}
+				}
+			}
+		}
+	}
+
 	// Find all unique materials and remap section to unique materials
 	TArray<UMaterialInterface*> UniqueMaterials;
 	TMap<UMaterialInterface*, int32> MaterialIndices;
@@ -3066,6 +3107,16 @@ void FMeshMergeUtilities::MergeComponentsToStaticMesh(const TArray<UPrimitiveCom
 			}
 		}
 
+		// Add merged sockets
+		if (InSettings.bMergeMeshSockets)
+		{
+			for (auto& [SocketName, Socket] : MergedSockets)
+			{
+				Socket->Rename(nullptr, StaticMesh);
+				StaticMesh->AddSocket(Socket);
+			}
+		}
+
 		StaticMesh->GetSectionInfoMap().CopyFrom(SectionInfoMap);
 		StaticMesh->GetOriginalSectionInfoMap().CopyFrom(SectionInfoMap);
 
diff --git a/Engine/Source/Runtime/Engine/Classes/Engine/MeshMerging.h b/Engine/Source/Runtime/Engine/Classes/Engine/MeshMerging.h
--- a/Engine/Source/Runtime/Engine/Classes/Engine/MeshMerging.h
+++ b/Engine/Source/Runtime/Engine/Classes/Engine/MeshMerging.h
@@ -522,6 +522,10 @@ struct FMeshMergingSettings
 	UPROPERTY(EditAnywhere, BlueprintReadWrite, Category = MeshSettings)
 	uint8 bMergePhysicsData:1;
 
+	/** Whether to merge sockets */
+	UPROPERTY(EditAnywhere, BlueprintReadWrite, Category = MeshSettings)
+	uint8 bMergeMeshSockets : 1;
+
 	/** Whether to merge source materials into one flat material, ONLY available when LOD Selection Type is set to LowestDetailLOD */
 	UPROPERTY(EditAnywhere, BlueprintReadWrite, Category = MaterialSettings, meta=(EditCondition="LODSelectionType == EMeshLODSelectionType::LowestDetailLOD"))
 	uint8 bMergeMaterials:1;
@@ -614,6 +618,7 @@ struct FMeshMergingSettings
 		, bComputedLightMapResolution(false)
 		, bPivotPointAtZero(false)
 		, bMergePhysicsData(false)
+		, bMergeMeshSockets(false)
 		, bMergeMaterials(false)
 		, bCreateMergedMaterial(false)
 		, bBakeVertexDataToMesh(false)
