Improve MeshMergeHelpers by adding option to apply component transform for static and skeletal meshes

diff --git a/Engine/Source/Developer/MeshMergeUtilities/Private/MeshMergeHelpers.cpp b/Engine/Source/Developer/MeshMergeUtilities/Private/MeshMergeHelpers.cpp
--- a/Engine/Source/Developer/MeshMergeUtilities/Private/MeshMergeHelpers.cpp
+++ b/Engine/Source/Developer/MeshMergeUtilities/Private/MeshMergeHelpers.cpp
@@ -231,7 +231,7 @@ void FMeshMergeHelpers::ExpandInstances(const UInstancedStaticMeshComponent* InI
 }
 
 
-static void RetrieveMeshInternal(const UStaticMesh* StaticMesh, const UStaticMeshComponent* StaticMeshComponent, int32 LODIndex, FMeshDescription& OutMeshDescription, bool bPropagateVertexColours)
+static void RetrieveMeshInternal(const UStaticMesh* StaticMesh, const UStaticMeshComponent* StaticMeshComponent, int32 LODIndex, FMeshDescription& OutMeshDescription, bool bPropagateVertexColours, bool bApplyComponentTransform)
 {
 	TRACE_CPUPROFILER_EVENT_SCOPE(FMeshMergeHelpers::RetrieveMeshInternal)
 
@@ -249,9 +249,6 @@ static void RetrieveMeshInternal(const UStaticMesh* StaticMesh, const UStaticMes
 	// If we have a component, use it to retrieve transform & vertex colors (if requested)
 	if (StaticMeshComponent)
 	{
-		// Transform mesh to world space
-		FTransform ComponentToWorldTransform = StaticMeshComponent->GetComponentTransform();
-
 		// Handle spline mesh deformation
 		const bool bIsSplineMeshComponent = StaticMeshComponent->IsA<USplineMeshComponent>();
 		if (bIsSplineMeshComponent)
@@ -267,8 +264,11 @@ static void RetrieveMeshInternal(const UStaticMesh* StaticMesh, const UStaticMes
 			FMeshMergeHelpers::PropagatePaintedColorsToMesh(StaticMeshComponent, LODIndex, OutMeshDescription);
 		}
 
-		// Transform raw mesh vertex data by the Static Mesh Component's component to world transformation	
-		FStaticMeshOperations::ApplyTransform(OutMeshDescription, ComponentToWorldTransform);
+		// If specified transform vertices from local to world space
+		if (bApplyComponentTransform)
+		{
+			FStaticMeshOperations::ApplyTransform(OutMeshDescription, StaticMeshComponent->GetComponentTransform());
+		}
 	}
 
 	FMeshBuildSettings BuildSettings;
@@ -300,14 +300,14 @@ static void RetrieveMeshInternal(const UStaticMesh* StaticMesh, const UStaticMes
 	FStaticMeshOperations::RecomputeNormalsAndTangentsIfNeeded(OutMeshDescription, ComputeNTBsOptions);
 }
 
-void FMeshMergeHelpers::RetrieveMesh(const UStaticMeshComponent* StaticMeshComponent, int32 LODIndex, FMeshDescription& OutMeshDescription, bool bPropagateVertexColours)
+void FMeshMergeHelpers::RetrieveMesh(const UStaticMeshComponent* StaticMeshComponent, int32 LODIndex, FMeshDescription& OutMeshDescription, bool bPropagateVertexColours, bool bApplyComponentTransform)
 {
-	RetrieveMeshInternal(StaticMeshComponent->GetStaticMesh(), StaticMeshComponent, LODIndex, OutMeshDescription, bPropagateVertexColours);
+	RetrieveMeshInternal(StaticMeshComponent->GetStaticMesh(), StaticMeshComponent, LODIndex, OutMeshDescription, bPropagateVertexColours, bApplyComponentTransform);
 }
 
 void FMeshMergeHelpers::RetrieveMesh(const UStaticMesh* StaticMesh, int32 LODIndex, FMeshDescription& OutMeshDescription)
 {
-	RetrieveMeshInternal(StaticMesh, nullptr, LODIndex, OutMeshDescription, /*bPropagateVertexColours*/false);
+	RetrieveMeshInternal(StaticMesh, nullptr, LODIndex, OutMeshDescription, /*bPropagateVertexColours*/false, /*bApplyComponentTransform*/false);
 }
 
 void FMeshMergeHelpers::ExportStaticMeshLOD(const FStaticMeshLODResources& StaticMeshLOD, FMeshDescription& OutMeshDescription, const TArray<FStaticMaterial>& Materials)
@@ -423,7 +423,7 @@ void FMeshMergeHelpers::ExportStaticMeshLOD(const FStaticMeshLODResources& Stati
 	}
 }
 
-void FMeshMergeHelpers::RetrieveMesh(const USkeletalMeshComponent* SkeletalMeshComponent, int32 LODIndex, FMeshDescription& OutMeshDescription, bool bPropagateVertexColours)
+void FMeshMergeHelpers::RetrieveMesh(const USkeletalMeshComponent* SkeletalMeshComponent, int32 LODIndex, FMeshDescription& OutMeshDescription, bool bPropagateVertexColours, bool bApplyComponentTransform)
 {
 	FSkeletalMeshModel* Resource = SkeletalMeshComponent->GetSkeletalMeshAsset()->GetImportedModel();
 	if (Resource->LODModels.IsValidIndex(LODIndex))
@@ -535,6 +535,12 @@ void FMeshMergeHelpers::RetrieveMesh(const USkeletalMeshComponent* SkeletalMeshC
 				const FPolygonID NewPolygonID = OutMeshDescription.CreatePolygon(SectionPolygonGroupID, VertexInstanceIDs);
 			}
 		}
+
+		// If specified transform vertices from local to world space
+		if (bApplyComponentTransform)
+		{
+			FStaticMeshOperations::ApplyTransform(OutMeshDescription, SkeletalMeshComponent->GetComponentTransform());
+		}
 	}
 }
 
diff --git a/Engine/Source/Developer/MeshMergeUtilities/Private/MeshMergeHelpers.h b/Engine/Source/Developer/MeshMergeUtilities/Private/MeshMergeHelpers.h
--- a/Engine/Source/Developer/MeshMergeUtilities/Private/MeshMergeHelpers.h
+++ b/Engine/Source/Developer/MeshMergeUtilities/Private/MeshMergeHelpers.h
@@ -31,8 +31,8 @@ public:
 	static void ExpandInstances(const UInstancedStaticMeshComponent* InInstancedStaticMeshComponent, FMeshDescription& InOutRawMesh);
 
 	/** Extracting mesh data in FMeshDescription form from static, skeletal mesh (components) */
-	static void RetrieveMesh(const UStaticMeshComponent* StaticMeshComponent, int32 LODIndex, FMeshDescription& OutMeshDescription, bool bPropagateVertexColours);
-	static void RetrieveMesh(const USkeletalMeshComponent* SkeletalMeshComponent, int32 LODIndex, FMeshDescription& OutMeshDescription, bool bPropagateVertexColours);
+	static void RetrieveMesh(const UStaticMeshComponent* StaticMeshComponent, int32 LODIndex, FMeshDescription& OutMeshDescription, bool bPropagateVertexColours, bool bApplyComponentTransform = true);
+	static void RetrieveMesh(const USkeletalMeshComponent* SkeletalMeshComponent, int32 LODIndex, FMeshDescription& OutMeshDescription, bool bPropagateVertexColours, bool bApplyComponentTransform = false);
 	static void RetrieveMesh(const UStaticMesh* StaticMesh, int32 LODIndex, FMeshDescription& OutMeshDescription);
 	
 	/** Exports static mesh LOD render data to a RawMesh */
