MeshMergeUtility - Fix a bug where collisions of ISM components were not merged correctly.

We new merge the collisions of all the instances in the ISM component, instead of a single collision instance using just the component transform.
Also exposing a utility function for getting the collision information of a list of components.

diff --git a/Engine/Source/Developer/MeshMergeUtilities/Private/MeshMergeUtilities.cpp b/Engine/Source/Developer/MeshMergeUtilities/Private/MeshMergeUtilities.cpp
--- a/Engine/Source/Developer/MeshMergeUtilities/Private/MeshMergeUtilities.cpp
+++ b/Engine/Source/Developer/MeshMergeUtilities/Private/MeshMergeUtilities.cpp
@@ -2300,7 +2300,7 @@ void FMeshMergeUtilities::MergeComponentsToStaticMesh(const TArray<UPrimitiveCom
 	TArray<FKAggregateGeom> PhysicsGeometry;
 	if (InSettings.bMergePhysicsData)
 	{
-		ExtractPhysicsDataFromComponents(ComponentsToMerge, PhysicsGeometry, BodySetupSource);
+		RetrievePhysicsData(ComponentsToMerge, PhysicsGeometry, BodySetupSource);
 	}
 
 	// Merge sockets
@@ -3698,16 +3698,52 @@ UMaterialInterface* FMeshMergeUtilities::CreateProxyMaterial(const FString &InBa
 	return MergedMaterial;
 }
 
-void FMeshMergeUtilities::ExtractPhysicsDataFromComponents(const TArray<UPrimitiveComponent*>& ComponentsToMerge, TArray<FKAggregateGeom>& InOutPhysicsGeometry, UBodySetup*& OutBodySetupSource) const
+void FMeshMergeUtilities::RetrievePhysicsData(const TArray<UPrimitiveComponent*>& ComponentsToMerge, TArray<FKAggregateGeom>& InOutPhysicsGeometry, UBodySetup*& OutBodySetupSource) const
 {
 	InOutPhysicsGeometry.AddDefaulted(ComponentsToMerge.Num());
-	for (int32 ComponentIndex = 0; ComponentIndex < ComponentsToMerge.Num(); ++ComponentIndex)
+	for (int32 ComponentIndex = 0, PhysicsGeometryIndex = 0; ComponentIndex < ComponentsToMerge.Num(); ++ComponentIndex)
 	{
 		UPrimitiveComponent* PrimComp = ComponentsToMerge[ComponentIndex];
 		UBodySetup* BodySetup = nullptr;
 		FTransform ComponentToWorld = FTransform::Identity;
 
-		if (UStaticMeshComponent* StaticMeshComp = Cast<UStaticMeshComponent>(PrimComp))
+		auto ExtractPhysicGeometry = [&BodySetup, &ComponentToWorld, &OutBodySetupSource, &InOutPhysicsGeometry, PrimComp](int32 PhysicsIndex) {
+				USplineMeshComponent* SplineMeshComponent = Cast<USplineMeshComponent>(PrimComp);
+				FMeshMergeHelpers::ExtractPhysicsGeometry(BodySetup, ComponentToWorld, SplineMeshComponent != nullptr, InOutPhysicsGeometry[PhysicsIndex]);
+				if (SplineMeshComponent)
+				{
+					FMeshMergeHelpers::PropagateSplineDeformationToPhysicsGeometry(SplineMeshComponent, InOutPhysicsGeometry[PhysicsIndex]);
+				}
+
+				// We will use first valid BodySetup as a source of physics settings
+				if (OutBodySetupSource == nullptr)
+				{
+					OutBodySetupSource = BodySetup;
+				}
+			};
+
+		if (UInstancedStaticMeshComponent* ISMComp = Cast<UInstancedStaticMeshComponent>(PrimComp))
+		{
+			const int32 NumberOfInstances = ISMComp->PerInstanceSMData.Num();
+			const UStaticMesh* SrcMesh = ISMComp->GetStaticMesh();
+			
+			if (NumberOfInstances > 1)
+			{
+				InOutPhysicsGeometry.AddDefaulted(NumberOfInstances - 1);
+			}
+
+			if (SrcMesh)
+			{
+				BodySetup = SrcMesh->GetBodySetup();
+			}
+
+			for (const FInstancedStaticMeshInstanceData& InstanceData : ISMComp->PerInstanceSMData)
+			{
+				ComponentToWorld = FTransform(InstanceData.Transform) * ISMComp->GetComponentToWorld();
+				ExtractPhysicGeometry(PhysicsGeometryIndex++);
+			}
+		}
+		else if (UStaticMeshComponent* StaticMeshComp = Cast<UStaticMeshComponent>(PrimComp))
 		{
 			UStaticMesh* SrcMesh = StaticMeshComp->GetStaticMesh();
 			if (SrcMesh)
@@ -3715,25 +3751,15 @@ void FMeshMergeUtilities::ExtractPhysicsDataFromComponents(const TArray<UPrimiti
 				BodySetup = SrcMesh->GetBodySetup();
 			}
 			ComponentToWorld = StaticMeshComp->GetComponentToWorld();
+			ExtractPhysicGeometry(PhysicsGeometryIndex++);
 		}
 		else if (UShapeComponent* ShapeComp = Cast<UShapeComponent>(PrimComp))
 		{
 			BodySetup = ShapeComp->GetBodySetup();
 			ComponentToWorld = ShapeComp->GetComponentToWorld();
+			ExtractPhysicGeometry(PhysicsGeometryIndex++);
 		}
 
-		USplineMeshComponent* SplineMeshComponent = Cast<USplineMeshComponent>(PrimComp);
-		FMeshMergeHelpers::ExtractPhysicsGeometry(BodySetup, ComponentToWorld, SplineMeshComponent != nullptr, InOutPhysicsGeometry[ComponentIndex]);
-		if (SplineMeshComponent)
-		{
-			FMeshMergeHelpers::PropagateSplineDeformationToPhysicsGeometry(SplineMeshComponent, InOutPhysicsGeometry[ComponentIndex]);
-		}
-
-		// We will use first valid BodySetup as a source of physics settings
-		if (OutBodySetupSource == nullptr)
-		{
-			OutBodySetupSource = BodySetup;
-		}
 	}
 }
 
diff --git a/Engine/Source/Developer/MeshMergeUtilities/Private/MeshMergeUtilities.h b/Engine/Source/Developer/MeshMergeUtilities/Private/MeshMergeUtilities.h
--- a/Engine/Source/Developer/MeshMergeUtilities/Private/MeshMergeUtilities.h
+++ b/Engine/Source/Developer/MeshMergeUtilities/Private/MeshMergeUtilities.h
@@ -43,12 +43,11 @@ public:
 	virtual void RetrieveMeshDescription(const USkeletalMeshComponent* InSkeletalMeshComponent, int32 LODIndex, FMeshDescription& InOutMeshDescription, bool bPropagateVertexColours) const override;
 	virtual void RetrieveMeshDescription(const UStaticMesh* InStaticMesh, int32 LODIndex, FMeshDescription& InOutMeshDescription) const override;
 
+	virtual void RetrievePhysicsData(const TArray<UPrimitiveComponent*>& InComponents, TArray<FKAggregateGeom>& InOutPhysicsGeometry, UBodySetup*& OutBodySetupSource) const override;
+
 	virtual void RegisterExtension(IMeshMergeExtension* InExtension) override;
 	virtual void UnregisterExtension(IMeshMergeExtension* InExtension) override;	
 protected:
-	/** Retrieves physics geometry and body setup from set of static mesh components */
-	void ExtractPhysicsDataFromComponents(const TArray<UPrimitiveComponent*>& ComponentsToMerge, TArray<FKAggregateGeom>& InOutPhysicsGeometry, UBodySetup*& OutBodySetupSource) const;
-
 	/** Determines whether or not an individual material uses model vertex data during the shading process and outputs per-material flags */
 	void DetermineMaterialVertexDataUsage(TArray<bool>& InOutMaterialUsesVertexData, const TArray<UMaterialInterface*>& UniqueMaterials, const UMaterialOptions* MaterialOptions) const;
 	/** Creates a proxy material instance at givne path and name */
diff --git a/Engine/Source/Developer/MeshMergeUtilities/Public/IMeshMergeUtilities.h b/Engine/Source/Developer/MeshMergeUtilities/Public/IMeshMergeUtilities.h
--- a/Engine/Source/Developer/MeshMergeUtilities/Public/IMeshMergeUtilities.h
+++ b/Engine/Source/Developer/MeshMergeUtilities/Public/IMeshMergeUtilities.h
@@ -135,6 +135,15 @@ public:
 	/** Extract mesh data in FMeshDescription form from the provided static mesh */
 	virtual void RetrieveMeshDescription(const UStaticMesh* InStaticMesh, int32 LODIndex, FMeshDescription& InOutMeshDescription) const = 0;
 
+	/**
+	 * Extract physics geometry and body setup from set of primitive components.
+	 * 
+	 * @param InComponents			List of components to extract the physics geometry from.
+	 * @param InOutPhysicsGeometry	The list of extracted physics geometry.
+	 * @param OutBodySetupSource	The BodySetup extracted.
+	 */
+	virtual void RetrievePhysicsData(const TArray<UPrimitiveComponent*>& InComponents, TArray<FKAggregateGeom>& InOutPhysicsGeometry, UBodySetup*& OutBodySetupSource) const = 0;
+
 	virtual void RegisterExtension(IMeshMergeExtension* InExtension) = 0;
 	virtual void UnregisterExtension(IMeshMergeExtension* InExtension) = 0;
 };
\ No newline at end of file
