Mesh Merging - Lazily evaluate vertex data usage by the material, as it is costly and we can avoid it in a lot of cases

diff --git a/Engine/Source/Developer/MeshMergeUtilities/Private/MeshMergeUtilities.cpp b/Engine/Source/Developer/MeshMergeUtilities/Private/MeshMergeUtilities.cpp
--- a/Engine/Source/Developer/MeshMergeUtilities/Private/MeshMergeUtilities.cpp
+++ b/Engine/Source/Developer/MeshMergeUtilities/Private/MeshMergeUtilities.cpp
@@ -169,9 +169,6 @@ void FMeshMergeUtilities::BakeMaterialsForComponent(TArray<TWeakObjectPtr<UObjec
 		UniqueMaterialToUniqueSectionMap.Add(UniqueIndex, SectionIndex);
 	}
 
-	TArray<bool> bMaterialUsesVertexData;
-	DetermineMaterialVertexDataUsage(bMaterialUsesVertexData, UniqueMaterials, MaterialOptions);
-
 	TArray<FMeshData> GlobalMeshSettings;
 	TArray<FMaterialData> GlobalMaterialSettings;
 	TArray<TMap<uint32, uint32>> OutputMaterialsMap;
@@ -180,7 +177,7 @@ void FMeshMergeUtilities::BakeMaterialsForComponent(TArray<TWeakObjectPtr<UObjec
 	for (int32 MaterialIndex = 0; MaterialIndex < UniqueMaterials.Num(); ++MaterialIndex)
 	{
 		UMaterialInterface* Material = UniqueMaterials[MaterialIndex];
-		const bool bDoesMaterialUseVertexData = bMaterialUsesVertexData[MaterialIndex];
+
 		// Retrieve all sections using this material 
 		TArray<uint32> SectionIndices;
 		UniqueMaterialToUniqueSectionMap.MultiFind(MaterialIndex, SectionIndices);
@@ -502,31 +499,28 @@ void FMeshMergeUtilities::BakeMaterialsForMesh(UStaticMesh* StaticMesh) const
 	BakeMaterialsForComponent(Objects, &Adapter);
 }
 
-void FMeshMergeUtilities::DetermineMaterialVertexDataUsage(TArray<bool>& InOutMaterialUsesVertexData, const TArray<UMaterialInterface*>& UniqueMaterials, const UMaterialOptions* MaterialOptions) const
+
+static bool DetermineMaterialVertexDataUsage(UMaterialInterface* Material, const UMaterialOptions* MaterialOptions)
 {
 	TRACE_CPUPROFILER_EVENT_SCOPE(DetermineMaterialVertexDataUsage);
 
-	InOutMaterialUsesVertexData.SetNum(UniqueMaterials.Num());
-	for (int32 MaterialIndex = 0; MaterialIndex < UniqueMaterials.Num(); ++MaterialIndex)
+	for (const FPropertyEntry& Entry : MaterialOptions->Properties)
 	{
-		UMaterialInterface* Material = UniqueMaterials[MaterialIndex];
-		for (const FPropertyEntry& Entry : MaterialOptions->Properties)
+		// Don't have to check a property if the result is going to be constant anyway
+		if (!Entry.bUseConstantValue && Entry.Property != MP_MAX)
 		{
-			// Don't have to check a property if the result is going to be constant anyway
-			if (!Entry.bUseConstantValue && Entry.Property != MP_MAX)
-			{
-				int32 NumTextureCoordinates;
-				bool bUsesVertexData;
-				Material->AnalyzeMaterialProperty(Entry.Property, NumTextureCoordinates, bUsesVertexData);
+			int32 NumTextureCoordinates;
+			bool bUsesVertexData;
+			Material->AnalyzeMaterialProperty(Entry.Property, NumTextureCoordinates, bUsesVertexData);
 
-				if (bUsesVertexData || NumTextureCoordinates > 1)
-				{
-					InOutMaterialUsesVertexData[MaterialIndex] = true;
-					break;
-				}
+			if (bUsesVertexData || NumTextureCoordinates > 1)
+			{
+				return true;
 			}
 		}
 	}
+
+	return false;
 }
 
 void FMeshMergeUtilities::ConvertOutputToFlatMaterials(const TArray<FBakeOutput>& BakeOutputs, const TArray<FMaterialData>& MaterialData, TArray<FFlattenMaterial> &FlattenedMaterials) const
@@ -2397,8 +2391,19 @@ void FMeshMergeUtilities::MergeComponentsToStaticMesh(const TArray<UPrimitiveCom
 		UMaterialOptions* MaterialOptions = PopulateMaterialOptions(MaterialProxySettings);
 
 		// Check each material to see if the shader actually uses vertex data and collect flags
-		TArray<bool> bMaterialUsesVertexData;
-		DetermineMaterialVertexDataUsage(bMaterialUsesVertexData, UniqueMaterials, MaterialOptions);
+		TArray<TOptional<bool>> bMaterialUsesVertexData;
+		bMaterialUsesVertexData.SetNum(UniqueMaterials.Num());
+
+		// Deferred call, as this may not be required by all code paths and is pretty costly to compute
+		auto DoesMaterialUsesVertexData = [&](const int32 InMaterialIndex)
+		{
+			if (!bMaterialUsesVertexData[InMaterialIndex].IsSet())
+			{
+				bMaterialUsesVertexData[InMaterialIndex] = DetermineMaterialVertexDataUsage(UniqueMaterials[InMaterialIndex], MaterialOptions);
+			}
+
+			return bMaterialUsesVertexData[InMaterialIndex].GetValue();
+		};
 
 		TArray<FMeshData> GlobalMeshSettings;
 		TArray<FMaterialData> GlobalMaterialSettings;
@@ -2452,14 +2457,13 @@ void FMeshMergeUtilities::MergeComponentsToStaticMesh(const TArray<UPrimitiveCom
 				UMaterialInterface* Material = MaterialSectionIndexPair.Key;
 				const int32 MaterialIndex = UniqueMaterials.IndexOfByKey(Material);
 				const TArray<int32>& SectionIndices = MaterialSectionIndexPair.Value;
-				const bool bDoesMaterialUseVertexData = bMaterialUsesVertexData[MaterialIndex];
 
 				FMaterialData MaterialData;
 				MaterialData.Material = CollapsedMaterialMap.FindChecked(Material);
 				MaterialData.PropertySizes = PropertySizes;
 
 				FMeshData NewMeshData;
-				const bool bUseMeshData = InSettings.bCreateMergedMaterial || bGloballyRemapUVs || (InSettings.bUseVertexDataForBakingMaterial && (bDoesMaterialUseVertexData || bRequiresUniqueUVs));
+				const bool bUseMeshData = InSettings.bCreateMergedMaterial || bGloballyRemapUVs || (InSettings.bUseVertexDataForBakingMaterial && (bRequiresUniqueUVs || DoesMaterialUsesVertexData(MaterialIndex)));
 				if (bUseMeshData)
 				{
 					NewMeshData.Mesh = Key.GetMesh();
diff --git a/Engine/Source/Developer/MeshMergeUtilities/Private/MeshMergeUtilities.h b/Engine/Source/Developer/MeshMergeUtilities/Private/MeshMergeUtilities.h
--- a/Engine/Source/Developer/MeshMergeUtilities/Private/MeshMergeUtilities.h
+++ b/Engine/Source/Developer/MeshMergeUtilities/Private/MeshMergeUtilities.h
@@ -48,8 +48,6 @@ public:
 	virtual void RegisterExtension(IMeshMergeExtension* InExtension) override;
 	virtual void UnregisterExtension(IMeshMergeExtension* InExtension) override;	
 protected:
-	/** Determines whether or not an individual material uses model vertex data during the shading process and outputs per-material flags */
-	void DetermineMaterialVertexDataUsage(TArray<bool>& InOutMaterialUsesVertexData, const TArray<UMaterialInterface*>& UniqueMaterials, const UMaterialOptions* MaterialOptions) const;
 	/** Creates a proxy material instance at givne path and name */
 	UMaterialInterface* CreateProxyMaterial(const FString &InBasePackageName, FString MergedAssetPackageName, UMaterialInterface* InBaseMaterial, UPackage* InOuter, const FMeshMergingSettings &InSettings, FFlattenMaterial OutMaterial, TArray<UObject *>& OutAssetsToSync, FMaterialUpdateContext* InMaterialUpdateContext = nullptr) const;
 	/** Merges flattened material into atlas textures */
