MeshMergeUtilities cleanup
* Moved material creation in MergeComponentsToStaticMesh() to a standalone utility function
* Deprecate bCreateMergedMaterial

diff --git a/Engine/Source/Developer/MaterialUtilities/Public/MaterialUtilities.h b/Engine/Source/Developer/MaterialUtilities/Public/MaterialUtilities.h
--- a/Engine/Source/Developer/MaterialUtilities/Public/MaterialUtilities.h
+++ b/Engine/Source/Developer/MaterialUtilities/Public/MaterialUtilities.h
@@ -3,6 +3,7 @@
 #pragma once
 
 #include "CoreMinimal.h"
+#include "Algo/AnyOf.h"
 #include "UObject/ObjectMacros.h"
 #include "Misc/Guid.h"
 #include "SceneTypes.h"
@@ -70,6 +71,11 @@ struct FFlattenMaterial
 		}
 	}
 
+	bool HasData() const
+	{
+		return Algo::AnyOf(PropertySamples, [](const TArray<FColor>& Samples) { return !Samples.IsEmpty(); });
+	}
+
 	/** Release all the sample data */
 	void ReleaseData()
 	{
diff --git a/Engine/Source/Developer/MeshMergeUtilities/Private/MeshMergeUtilities.cpp b/Engine/Source/Developer/MeshMergeUtilities/Private/MeshMergeUtilities.cpp
--- a/Engine/Source/Developer/MeshMergeUtilities/Private/MeshMergeUtilities.cpp
+++ b/Engine/Source/Developer/MeshMergeUtilities/Private/MeshMergeUtilities.cpp
@@ -48,7 +48,6 @@
 
 #include "ProxyGenerationProcessor.h"
 #include "IMaterialBakingAdapter.h"
-#include "StaticMeshComponentAdapter.h"
 #include "StaticMeshComponentLODInfo.h"
 #include "SkeletalMeshAdapter.h"
 #include "StaticMeshAdapter.h"
@@ -159,7 +158,6 @@ void FMeshMergeUtilities::BakeMaterialsForComponent(TArray<TWeakObjectPtr<UObjec
 	}
 
 	TArray<UMaterialInterface*> UniqueMaterials;
-	TMap<UMaterialInterface*, int32> MaterialIndices;
 	TMultiMap<uint32, uint32> UniqueMaterialToUniqueSectionMap;
 	// Populate list of unique materials and store section mappings
 	for (int32 SectionIndex = 0; SectionIndex < UniqueSections.Num(); ++SectionIndex)
@@ -2289,14 +2287,6 @@ void FMeshMergeUtilities::MergeComponentsToStaticMesh(const TArray<UPrimitiveCom
 		DataTracker.ProcessRawMeshes();
 	}
 
-	// Retrieve physics data
-	UBodySetup* BodySetupSource = nullptr;
-	TArray<FKAggregateGeom> PhysicsGeometry;
-	if (InSettings.bMergePhysicsData)
-	{
-		RetrievePhysicsData(ComponentsToMerge, PhysicsGeometry, BodySetupSource);
-	}
-
 	// Merge sockets
 	TMap<FName, UStaticMeshSocket*> MergedSockets;
 	if (InSettings.bMergeMeshSockets)
@@ -2339,7 +2329,6 @@ void FMeshMergeUtilities::MergeComponentsToStaticMesh(const TArray<UPrimitiveCom
 
 	// Find all unique materials and remap section to unique materials
 	TArray<UMaterialInterface*> UniqueMaterials;
-	TMap<UMaterialInterface*, int32> MaterialIndices;
 	TMap<UMaterialInterface*, UMaterialInterface*> CollapsedMaterialMap;
 
 	for (int32 SectionIndex = 0; SectionIndex < DataTracker.NumberOfUniqueSections(); ++SectionIndex)
@@ -2347,16 +2336,16 @@ void FMeshMergeUtilities::MergeComponentsToStaticMesh(const TArray<UPrimitiveCom
 		// Unique index for material
 		UMaterialInterface* MaterialInterface = DataTracker.GetMaterialForSectionIndex(SectionIndex);
 		int32 UniqueIndex = UniqueMaterials.IndexOfByPredicate([&InSettings, MaterialInterface](const UMaterialInterface* InMaterialInterface)
-		{
-			// Perform an optional custom comparison if we are trying to collapse material instances
-			if(InSettings.bMergeEquivalentMaterials)
 			{
-				return FMaterialKey(MaterialInterface) == FMaterialKey(InMaterialInterface);
-			}
-			return MaterialInterface == InMaterialInterface;
-		});
+				// Perform an optional custom comparison if we are trying to collapse material instances
+				if (InSettings.bMergeEquivalentMaterials)
+				{
+					return FMaterialKey(MaterialInterface) == FMaterialKey(InMaterialInterface);
+				}
+				return MaterialInterface == InMaterialInterface;
+			});
 
-		if(UniqueIndex == INDEX_NONE)
+		if (UniqueIndex == INDEX_NONE)
 		{
 			UniqueIndex = UniqueMaterials.Add(MaterialInterface);
 		}
@@ -2365,475 +2354,37 @@ void FMeshMergeUtilities::MergeComponentsToStaticMesh(const TArray<UPrimitiveCom
 		CollapsedMaterialMap.Add(MaterialInterface, UniqueMaterials[UniqueIndex]);
 	}
 
-	// For each unique material calculate how 'important' they are
-	TArray<float> MaterialImportanceValues;
-	FMaterialUtilities::DetermineMaterialImportance(UniqueMaterials, MaterialImportanceValues);
+	// Retrieve physics data
+	UBodySetup* BodySetupSource = nullptr;
+	TArray<FKAggregateGeom> PhysicsGeometry;
+	if (InSettings.bMergePhysicsData)
+	{
+		RetrievePhysicsData(ComponentsToMerge, PhysicsGeometry, BodySetupSource);
+	}
 
 	TMultiMap< FMeshLODKey, MaterialRemapPair > OutputMaterialsMap;
 
-	// The UV channel to use for the flattened material
-	int32 MergedMatUVChannel = INDEX_NONE;
-	UMaterialInterface* MergedMaterial = nullptr;
-
 	// If the user wants to merge materials into a single one
 	if (bMergeMaterialData && UniqueMaterials.Num() != 0)
 	{
-		TArray<UPrimitiveComponent*> PrimitiveComponents;
-		Algo::Transform(StaticMeshComponentsToMerge, PrimitiveComponents, [](UStaticMeshComponent* SMComponent) { return SMComponent; });
-
-		FMaterialProxySettings MaterialProxySettings = InSettings.MaterialSettings;
-		if (MaterialProxySettings.ResolveTexelDensity(PrimitiveComponents))
-		{
-			MaterialProxySettings.TextureSize = DataTracker.GetTextureSizeFromTargetTexelDensity(MaterialProxySettings.TargetTexelDensityPerMeter);
-			MaterialProxySettings.TextureSizingType = ETextureSizingType::TextureSizingType_UseSingleTextureSize;
-		}
-
-		UMaterialOptions* MaterialOptions = PopulateMaterialOptions(MaterialProxySettings);
-
-		// Check each material to see if the shader actually uses vertex data and collect flags
-		TArray<TOptional<bool>> bMaterialUsesVertexData;
-		bMaterialUsesVertexData.SetNum(UniqueMaterials.Num());
-
-		// Deferred call, as this may not be required by all code paths and is pretty costly to compute
-		auto DoesMaterialUsesVertexData = [&](const int32 InMaterialIndex)
-		{
-			if (!bMaterialUsesVertexData[InMaterialIndex].IsSet())
-			{
-				bMaterialUsesVertexData[InMaterialIndex] = DetermineMaterialVertexDataUsage(UniqueMaterials[InMaterialIndex], MaterialOptions);
-			}
-
-			return bMaterialUsesVertexData[InMaterialIndex].GetValue();
-		};
-
-		TArray<FMeshData> GlobalMeshSettings;
-		TArray<FMaterialData> GlobalMaterialSettings;
-		TArray<float> SectionMaterialImportanceValues;
-
-		TMap<EMaterialProperty, FIntPoint> PropertySizes;
-		for (const FPropertyEntry& Entry : MaterialOptions->Properties)
-		{
-			if (!Entry.bUseConstantValue && Entry.Property != MP_MAX)
-			{
-				PropertySizes.Add(Entry.Property, Entry.bUseCustomSize ? Entry.CustomSize : MaterialOptions->TextureSize);
-			}
-		}
-
-		// If we are generating a single LOD and want to merge materials we can utilize texture space better by generating unique UVs
-		// for the merged mesh and baking out materials using those UVs
-		const bool bGloballyRemapUVs = !bMergeAllLODs && !InSettings.bReuseMeshLightmapUVs;
-
-		typedef TTuple<UStaticMesh*, int32> FMeshLODTuple;
-		typedef TFuture<TArray<FVector2D>> FUVComputeFuture;
-		TMap<FMeshLODTuple, FUVComputeFuture> MeshLODsTextureCoordinates;
-		TMap<int32, FMeshLODTuple> MeshDataAwaitingResults;
-
-		for (TConstRawMeshIterator RawMeshIterator = DataTracker.GetConstRawMeshIterator(); RawMeshIterator; ++RawMeshIterator)
-		{
-			const FMeshLODKey& Key = RawMeshIterator.Key();
-			const FMeshDescription& RawMesh = RawMeshIterator.Value();
-			const bool bRequiresUniqueUVs = DataTracker.DoesMeshLODRequireUniqueUVs(Key);
-
-			FMeshDescription* MeshDescription = DataTracker.GetRawMeshPtr(Key);
-			UStaticMeshComponent* Component = StaticMeshComponentsToMerge[Key.GetMeshIndex()];
-			UStaticMesh* StaticMesh = Component->GetStaticMesh();
-
-			// Retrieve all sections and materials for key
-			TArray<SectionRemapPair> SectionRemapPairs;
-			DataTracker.GetMappingsForMeshLOD(Key, SectionRemapPairs);
-
-			// Contains unique materials used for this key, and the accompanying section index which point to the material
-			TMap<UMaterialInterface*, TArray<int32>> MaterialAndSectionIndices;
-
-			for (const SectionRemapPair& RemapPair : SectionRemapPairs)
-			{
-				const int32 UniqueIndex = RemapPair.Value;
-				const int32 SectionIndex = RemapPair.Key;
-				TArray<int32>& SectionIndices = MaterialAndSectionIndices.FindOrAdd(CollapsedMaterialMap.FindChecked(DataTracker.GetMaterialForSectionIndex(UniqueIndex)));
-				SectionIndices.Add(SectionIndex);
-			}
-
-			for (TPair<UMaterialInterface*, TArray<int32>>& MaterialSectionIndexPair : MaterialAndSectionIndices)
-			{
-				UMaterialInterface* Material = MaterialSectionIndexPair.Key;
-				const int32 MaterialIndex = UniqueMaterials.IndexOfByKey(Material);
-				const TArray<int32>& SectionIndices = MaterialSectionIndexPair.Value;
-
-				FMaterialData MaterialData;
-				MaterialData.Material = CollapsedMaterialMap.FindChecked(Material);
-				MaterialData.PropertySizes = PropertySizes;
-
-				FMeshData NewMeshData;
-				const bool bUseMeshData = InSettings.bCreateMergedMaterial || bGloballyRemapUVs || (InSettings.bUseVertexDataForBakingMaterial && (bRequiresUniqueUVs || DoesMaterialUsesVertexData(MaterialIndex)));
-				if (bUseMeshData)
-				{
-					NewMeshData.Mesh = Key.GetMesh();
-					NewMeshData.MeshDescription = MeshDescription;
-					NewMeshData.VertexColorHash = Key.GetVertexColorHash();
-					NewMeshData.bMirrored = Component->GetComponentTransform().GetDeterminant() < 0.0f;					
-					NewMeshData.MaterialIndices = SectionIndices;
-				}
-
-				auto CompareMaterialData = [&InSettings](const FMaterialData& LHS, const FMaterialData& RHS)
-				{
-					return InSettings.bMergeEquivalentMaterials ? FMaterialKey(LHS.Material) == FMaterialKey(RHS.Material) : LHS.Material == RHS.Material;
-				};
-
-				auto CompareMeshData = [](const FMeshData& LHS, const FMeshData& RHS)
-				{
-					return (LHS.Mesh == RHS.Mesh) && (LHS.MaterialIndices == RHS.MaterialIndices) && (LHS.bMirrored == RHS.bMirrored) && (LHS.VertexColorHash == RHS.VertexColorHash);
-				};
-
-				// Find material & mesh pair
-				int32 MeshDataIndex = INDEX_NONE;
-				for (int32 GlobalMaterialSettingsIndex = 0; GlobalMaterialSettingsIndex < GlobalMaterialSettings.Num(); ++GlobalMaterialSettingsIndex)
-				{
-					if (CompareMaterialData(GlobalMaterialSettings[GlobalMaterialSettingsIndex], MaterialData) &&
-						CompareMeshData(GlobalMeshSettings[GlobalMaterialSettingsIndex], NewMeshData))
-					{
-						MeshDataIndex = GlobalMaterialSettingsIndex;
-						break;
-					}
-				}
-
-				// We've found a match, no need to process this mesh/material pair
-				if (MeshDataIndex == INDEX_NONE)
-				{
-					// We're processing a new pair
-					MeshDataIndex = GlobalMeshSettings.Num();
-
-					FMeshData& MeshData = GlobalMeshSettings.Emplace_GetRef(NewMeshData);
-					GlobalMaterialSettings.Add(MaterialData);
-					SectionMaterialImportanceValues.Add(MaterialImportanceValues[MaterialIndex]);
-				
-					if (bUseMeshData)
-					{
-						// if it has vertex color/*WedgetColors.Num()*/, it should also use light map UV index
-						// we can't do this for all meshes, but only for the mesh that has vertex color.
-						if (bRequiresUniqueUVs || MeshData.MeshDescription->VertexInstances().Num() > 0)
-						{
-							// Check if there are lightmap uvs available?
-							const int32 LightMapUVIndex = StaticMesh->GetLightMapCoordinateIndex();
-
-							TVertexInstanceAttributesConstRef<FVector2f> VertexInstanceUVs = FStaticMeshConstAttributes(*MeshData.MeshDescription).GetVertexInstanceUVs();
-							if (InSettings.bReuseMeshLightmapUVs && VertexInstanceUVs.GetNumElements() > 0 && VertexInstanceUVs.GetNumChannels() > LightMapUVIndex)
-							{
-								MeshData.TextureCoordinateIndex = LightMapUVIndex;
-							}
-							else
-							{
-								// Verify if we started an async task to generate UVs for this static mesh & LOD
-								FMeshLODTuple Tuple(Key.GetMesh(), Key.GetLODIndex());
-								if (!MeshLODsTextureCoordinates.Find(Tuple))
-								{
-									// No job found yet, fire an async task
-									MeshLODsTextureCoordinates.Add(Tuple, Async(EAsyncExecution::Thread, [MeshDescription, MaterialOptions, this]()
-									{
-										TArray<FVector2D> UniqueTextureCoordinates;
-										FStaticMeshOperations::GenerateUniqueUVsForStaticMesh(*MeshDescription, MaterialOptions->TextureSize.GetMax(), false, UniqueTextureCoordinates);
-										ScaleTextureCoordinatesToBox(FBox2D(FVector2D::ZeroVector, FVector2D(1, 1)), UniqueTextureCoordinates);
-										return UniqueTextureCoordinates;
-									}));
-								}
-								// Keep track of the fact that this mesh is waiting for the UV computation to finish
-								MeshDataAwaitingResults.Add(MeshDataIndex, Tuple);
-							}
-						}
-
-						Adapters[Key.GetMeshIndex()].ApplySettings(Key.GetLODIndex(), MeshData);
-					}
-				}
-
-				for (const auto& OriginalSectionIndex : SectionIndices)
-				{
-					OutputMaterialsMap.Add(Key, MaterialRemapPair(OriginalSectionIndex, MeshDataIndex));
-				}
-			}
-		}
-
-		// Fetch results from the async UV computation tasks
-		for (auto MeshData : MeshDataAwaitingResults)
-		{
-			GlobalMeshSettings[MeshData.Key].CustomTextureCoordinates = MeshLODsTextureCoordinates[MeshData.Value].Get();
-		}
-
-		TArray<FMeshData*> MeshSettingPtrs;
-		for (int32 SettingsIndex = 0; SettingsIndex < GlobalMeshSettings.Num(); ++SettingsIndex)
-		{
-			MeshSettingPtrs.Add(&GlobalMeshSettings[SettingsIndex]);
-		}
-
-		TArray<FMaterialData*> MaterialSettingPtrs;
-		for (int32 SettingsIndex = 0; SettingsIndex < GlobalMaterialSettings.Num(); ++SettingsIndex)
-		{
-			MaterialSettingPtrs.Add(&GlobalMaterialSettings[SettingsIndex]);
-		}
-
-		if(bGloballyRemapUVs)
-		{
-			// We must keep vertex data in order to properly generate unique UVs
-			FMeshMergingSettings RemapUVMergeSettings = InSettings;
-			RemapUVMergeSettings.bBakeVertexDataToMesh = true;
-
-			TArray<FMeshDescription> MergedRawMeshes;
-			CreateMergedRawMeshes(DataTracker, RemapUVMergeSettings, StaticMeshComponentsToMerge, UniqueMaterials, CollapsedMaterialMap, OutputMaterialsMap, false, false, MergedAssetPivot, MergedRawMeshes);
-
-			// Create texture coords for the merged mesh
-			TArray<FVector2D> GlobalTextureCoordinates;
-			bool bSuccess = FStaticMeshOperations::GenerateUniqueUVsForStaticMesh(MergedRawMeshes[0], MaterialOptions->TextureSize.GetMax(), true, GlobalTextureCoordinates);
-			if (bSuccess)
-			{
-				ScaleTextureCoordinatesToBox(FBox2D(FVector2D::ZeroVector, FVector2D(1, 1)), GlobalTextureCoordinates);
-
-				// copy UVs back to the un-merged mesh's custom texture coords
-				// iterate the raw meshes in the same way as when we combined the mesh above in CreateMergedRawMeshes()
-				int32 GlobalUVIndex = 0;
-				for (TConstRawMeshIterator RawMeshIterator = DataTracker.GetConstRawMeshIterator(); RawMeshIterator; ++RawMeshIterator)
-				{
-					const FMeshLODKey& Key = RawMeshIterator.Key();
-					const FMeshDescription& RawMesh = RawMeshIterator.Value();
-
-					// Build a local array for this raw mesh
-					TArray<FVector2D> UniqueTextureCoordinates;
-					UniqueTextureCoordinates.SetNumUninitialized(RawMesh.VertexInstances().Num());
-					for (FVector2D& UniqueTextureCoordinate : UniqueTextureCoordinates)
-					{
-						UniqueTextureCoordinate = GlobalTextureCoordinates[GlobalUVIndex++];
-					}
-
-					// copy to mesh data
-					for (FMeshData& MeshData : GlobalMeshSettings)
-					{
-						if (MeshData.MeshDescription == &RawMesh)
-						{
-							MeshData.CustomTextureCoordinates = UniqueTextureCoordinates;
-						}
-					}
-				}
-
-				// Dont smear borders as we will copy back non-pink pixels
-				for (FMaterialData& MaterialData : GlobalMaterialSettings)
-				{
-					MaterialData.bPerformBorderSmear = false;
-				}
-			}
-			else
-			{
-				UE_LOG(LogMeshMerging, Warning, TEXT("GenerateUniqueUVsForStaticMesh: Failed to pack UVs for static mesh"));
-			}
-		}
-
-		TArray<FFlattenMaterial> FlattenedMaterials;
-		// This scope ensures BakeOutputs is never used after TransferOutputToFlatMaterials
-		{
-			TArray<FBakeOutput> BakeOutputs;
-			IMaterialBakingModule& Module = FModuleManager::Get().LoadModuleChecked<IMaterialBakingModule>("MaterialBaking");
-			Module.BakeMaterials(MaterialSettingPtrs, MeshSettingPtrs, BakeOutputs);
-
-			// Append constant properties ?
-			TArray<FColor> ConstantData;
-			FIntPoint ConstantSize(1, 1);
-			for (const FPropertyEntry& Entry : MaterialOptions->Properties)
-			{
-				if (Entry.bUseConstantValue && Entry.Property != MP_MAX)
-				{
-					ConstantData.SetNum(1, false);
-					ConstantData[0] = FLinearColor(Entry.ConstantValue, Entry.ConstantValue, Entry.ConstantValue).ToFColor(true);
-					for (FBakeOutput& Output : BakeOutputs)
-					{
-						Output.PropertyData.Add(Entry.Property, ConstantData);
-						Output.PropertySizes.Add(Entry.Property, ConstantSize);
-					}
-				}
-			}
-
-			TransferOutputToFlatMaterials(GlobalMaterialSettings, BakeOutputs, FlattenedMaterials);
-		}
-
-		if(!bGloballyRemapUVs)
-		{
-			// Try to optimize materials where possible	
-			for (FFlattenMaterial& InMaterial : FlattenedMaterials)
-			{
-				FMaterialUtilities::OptimizeFlattenMaterial(InMaterial);
-			}
-		}
-
-		FFlattenMaterial OutMaterial;
-		for (const FPropertyEntry& Entry : MaterialOptions->Properties)
-		{
-			if (Entry.Property != MP_MAX)
-			{
-				EFlattenMaterialProperties OldProperty = ToFlattenProperty(Entry.Property);
-				if (ensure(OldProperty != EFlattenMaterialProperties::NumFlattenMaterialProperties))
-				{
-					OutMaterial.SetPropertySize(OldProperty, Entry.bUseCustomSize ? Entry.CustomSize : MaterialOptions->TextureSize);
-				}
-			}
-		}
-
-		TArray<FUVOffsetScalePair> UVTransforms;
-		if(bGloballyRemapUVs)
-		{
-			// If we have globally remapped UVs we copy non-pink pixels over the dest texture rather than 
-			// copying sub-charts
-			TArray<FBox2D> MaterialBoxes;
-			MaterialBoxes.SetNumUninitialized(GlobalMaterialSettings.Num());
-			for(FBox2D& Box2D : MaterialBoxes)
-			{
-				Box2D = FBox2D(FVector2D(0.0f, 0.0f), FVector2D(1.0f, 1.0f));
-			}
-
-			FlattenBinnedMaterials(FlattenedMaterials, MaterialBoxes, 0, true, OutMaterial, UVTransforms);
-		}
-		else
-		{
-			/** Reweighting */
-			float TotalValue = 0.0f;
-			for (const float& Value : SectionMaterialImportanceValues)
-			{
-				TotalValue += Value;
-			}
-
-			float Multiplier = 1.0f / TotalValue;
-
-			for (float& Value : SectionMaterialImportanceValues)
-			{
-				Value *= Multiplier;
-			}
-			/** End reweighting */
-
-			if (InSettings.bUseTextureBinning)
-			{
-				TArray<FBox2D> MaterialBoxes;
-				FMaterialUtilities::GeneratedBinnedTextureSquares(FVector2D(1.0f, 1.0f), SectionMaterialImportanceValues, MaterialBoxes);
-				FlattenBinnedMaterials(FlattenedMaterials, MaterialBoxes, InSettings.GutterSize, false, OutMaterial, UVTransforms);
-			}
-			else
-			{
-				MergeFlattenedMaterials(FlattenedMaterials, InSettings.GutterSize, OutMaterial, UVTransforms);
-			}
-		}
-
-		// Compute UV channel to use for the merged material
-		if(InSettings.bCreateMergedMaterial)
-		{
-			for(TConstRawMeshIterator Iterator = DataTracker.GetConstRawMeshIterator(); Iterator; ++Iterator)
-			{
-				const FMeshDescription& RawMesh = Iterator.Value();
-
-				if(RawMesh.Vertices().Num())
-				{
-					TVertexInstanceAttributesRef<const FVector2f> VertexInstanceUVs = FStaticMeshConstAttributes(RawMesh).GetVertexInstanceUVs();
-					MergedMatUVChannel = FMath::Max(MergedMatUVChannel, VertexInstanceUVs.GetNumChannels());
-				}
-			}
-		}
-
-		// Adjust UVs
-		for (int32 ComponentIndex = 0; ComponentIndex < StaticMeshComponentsToMerge.Num(); ++ComponentIndex)
-		{
-			TArray<uint32> ProcessedMaterials;
-			for (TPair<FMeshLODKey, MaterialRemapPair>& MappingPair : OutputMaterialsMap)
-			{
-				if (MappingPair.Key.GetMeshIndex() == ComponentIndex && !ProcessedMaterials.Contains(MappingPair.Value.Key))
-				{
-					// Retrieve raw mesh data for this component and lod pair
-					FMeshDescription* RawMesh = DataTracker.GetRawMeshPtr(MappingPair.Key);
-
-					FMeshData& MeshData = GlobalMeshSettings[MappingPair.Value.Value];
-					const FUVOffsetScalePair& UVTransform = UVTransforms[MappingPair.Value.Value];
-
-					const uint32 MaterialIndex = MappingPair.Value.Key;
-					ProcessedMaterials.Add(MaterialIndex);
-					if (RawMesh->Vertices().Num())
-					{
-						TVertexInstanceAttributesRef<FVector2f> VertexInstanceUVs = FStaticMeshAttributes(*RawMesh).GetVertexInstanceUVs();
-						int32 NumUVChannel = FMath::Min(VertexInstanceUVs.GetNumChannels(), (int32)MAX_MESH_TEXTURE_COORDS);
-						for (int32 UVChannelIdx = 0; UVChannelIdx < NumUVChannel; ++UVChannelIdx)
-						{
-							int32 VertexIndex = 0;
-							for (FVertexInstanceID VertexInstanceID : RawMesh->VertexInstances().GetElementIDs())
-							{
-								FVector2D UV = FVector2D(VertexInstanceUVs.Get(VertexInstanceID, UVChannelIdx));
-								if (UVChannelIdx == 0 && !InSettings.bCreateMergedMaterial)
-								{
-									if (MeshData.CustomTextureCoordinates.Num())
-									{
-										UV = MeshData.CustomTextureCoordinates[VertexIndex];
-									}
-									else if (MeshData.TextureCoordinateIndex != 0)
-									{
-										check(MeshData.TextureCoordinateIndex < NumUVChannel);
-										UV = FVector2D(VertexInstanceUVs.Get(VertexInstanceID, MeshData.TextureCoordinateIndex));
-									}
-								}
-
-								const TArray<FPolygonID>& Polygons = RawMesh->GetVertexInstanceConnectedPolygons(VertexInstanceID);
-								for (FPolygonID PolygonID : Polygons)
-								{
-									FPolygonGroupID PolygonGroupID = RawMesh->GetPolygonPolygonGroup(PolygonID);
-									if (PolygonGroupID.GetValue() == MaterialIndex)
-									{
-										if (UVTransform.Value != FVector2D::ZeroVector)
-										{
-											VertexInstanceUVs.Set(VertexInstanceID, UVChannelIdx, FVector2f(UV * UVTransform.Value + UVTransform.Key));	// LWC_TODO: Precision loss
-											break;
-										}
-									}
-								}
-								VertexIndex++;
-							}
-						}
-
-						if (InSettings.bCreateMergedMaterial && MeshData.CustomTextureCoordinates.Num() > 0)
-						{
-							VertexInstanceUVs.SetNumChannels(MergedMatUVChannel + 1);
-
-							int32 VertexIndex = 0;
-							for(FVertexInstanceID VertexInstanceID : RawMesh->VertexInstances().GetElementIDs())
-							{
-								FVector2D UV = MeshData.CustomTextureCoordinates[VertexIndex];
-								VertexInstanceUVs.Set(VertexInstanceID, MergedMatUVChannel, FVector2f(UV * UVTransform.Value + UVTransform.Key));	// LWC_TODO: Precision loss
-								VertexIndex++;
-							}
-						}
-					}
-				}
-			}
-		}
-
-		for (TRawMeshIterator Iterator = DataTracker.GetRawMeshIterator(); Iterator; ++Iterator)
+		// Create the merged material
+		FFlattenMaterial FlattenMaterial;
+		CreateMergedMaterial(DataTracker, InSettings, StaticMeshComponentsToMerge, Adapters, UniqueMaterials, CollapsedMaterialMap, OutputMaterialsMap, bMergeAllLODs, bMergeMaterialData, MergedAssetPivot, FlattenMaterial);
+		if (FlattenMaterial.HasData())
 		{
-			FMeshDescription& RawMesh = Iterator.Value();
-			// Reset material indexes
-			TMap<FPolygonGroupID, FPolygonGroupID> RemapPolygonGroups;
-			for (FPolygonGroupID PolygonGroupID : RawMesh.PolygonGroups().GetElementIDs())
-			{
-				RemapPolygonGroups.Add(PolygonGroupID, FPolygonGroupID(0));
-			}
-			RawMesh.RemapPolygonGroups(RemapPolygonGroups);
-		}
+			// Don't recreate render states with the material update context as we will manually do it through
+			// the FStaticMeshComponentRecreateRenderStateContext used below at the creation of the static mesh.
+			FMaterialUpdateContext MaterialUpdateContext(FMaterialUpdateContext::EOptions::Default & ~FMaterialUpdateContext::EOptions::RecreateRenderStates);
 
-		OutMaterial.UVChannel = MergedMatUVChannel;
+			UMaterialInterface* MergedMaterial = CreateProxyMaterial(InBasePackageName, MergedAssetPackageName, InBaseMaterial, InOuter, InSettings, FlattenMaterial, OutAssetsToSync, &MaterialUpdateContext);
 
-		// Don't recreate render states with the material update context as we will manually do it through
-		// the FStaticMeshComponentRecreateRenderStateContext used below at the creation of the static mesh.
-		FMaterialUpdateContext MaterialUpdateContext(FMaterialUpdateContext::EOptions::Default & ~FMaterialUpdateContext::EOptions::RecreateRenderStates);
+			UniqueMaterials.Empty(1);
+			UniqueMaterials.Add(MergedMaterial);
 
-		MergedMaterial = CreateProxyMaterial(InBasePackageName, MergedAssetPackageName, InBaseMaterial, InOuter, InSettings, OutMaterial, OutAssetsToSync, &MaterialUpdateContext);
-		
-		if (MergedMaterial)
-		{
-			if (!InSettings.bCreateMergedMaterial)
-			{
-				UniqueMaterials.Empty(1);
-				UniqueMaterials.Add(MergedMaterial);
-
-				FSectionInfo NewSection;
-				NewSection.Material = MergedMaterial;
-				NewSection.EnabledProperties.Add(GET_MEMBER_NAME_CHECKED(FStaticMeshSection, bCastShadow));
-				DataTracker.AddBakedMaterialSection(NewSection);
-			}
+			FSectionInfo NewSection;
+			NewSection.Material = MergedMaterial;
+			NewSection.EnabledProperties.Add(GET_MEMBER_NAME_CHECKED(FStaticMeshSection, bCastShadow));
+			DataTracker.AddBakedMaterialSection(NewSection);
 
 			for (IMeshMergeExtension* Extension : MeshMergeExtensions)
 			{
@@ -2844,7 +2395,7 @@ void FMeshMergeUtilities::MergeComponentsToStaticMesh(const TArray<UPrimitiveCom
 
 	// Create the merged mesh
 	TArray<FMeshDescription> MergedRawMeshes;
-	CreateMergedRawMeshes(DataTracker, InSettings, StaticMeshComponentsToMerge, UniqueMaterials, CollapsedMaterialMap, OutputMaterialsMap, bMergeAllLODs, bMergeMaterialData && !InSettings.bCreateMergedMaterial, MergedAssetPivot, MergedRawMeshes);
+	CreateMergedRawMeshes(DataTracker, InSettings, StaticMeshComponentsToMerge, UniqueMaterials, CollapsedMaterialMap, OutputMaterialsMap, bMergeAllLODs, bMergeMaterialData, MergedAssetPivot, MergedRawMeshes);
 
 	// Notify listeners that our merged mesh was created
 	for (IMeshMergeExtension* Extension : MeshMergeExtensions)
@@ -2893,11 +2444,6 @@ void FMeshMergeUtilities::MergeComponentsToStaticMesh(const TArray<UPrimitiveCom
 		}
 	}
 
-	if(InSettings.bCreateMergedMaterial)
-	{
-		OutputMaterialsMap.Reset();
-	}
-
 	// Transform physics primitives to merged mesh pivot
 	if (InSettings.bMergePhysicsData && !MergedAssetPivot.IsZero())
 	{
@@ -3152,23 +2698,429 @@ void FMeshMergeUtilities::MergeComponentsToStaticMesh(const TArray<UPrimitiveCom
 			StaticMesh->CalculateExtendedBounds();
 		}		
 
-		if (InSettings.bCreateMergedMaterial && MergedMaterial)
+		StaticMesh->PostEditChange();
+
+		OutAssetsToSync.Add(StaticMesh);
+		OutMergedActorLocation = MergedAssetPivot;
+	}
+}
+
+void FMeshMergeUtilities::CreateMergedMaterial(FMeshMergeDataTracker& InDataTracker, const FMeshMergingSettings& InSettings, const TArray<UStaticMeshComponent*>& InStaticMeshComponentsToMerge, TArray<FStaticMeshComponentAdapter>& InAdapters, const TArray<UMaterialInterface*>& InUniqueMaterials, const TMap<UMaterialInterface*, UMaterialInterface*>& InCollapsedMaterialMap, TMultiMap<FMeshLODKey, MaterialRemapPair>& InOutputMaterialsMap, bool bInMergeAllLODs, bool bInMergeMaterialData, const FVector& InMergedAssetPivot, FFlattenMaterial& OutFlattenMaterial) const
+{
+	OutFlattenMaterial.ReleaseData();
+
+	TArray<UPrimitiveComponent*> PrimitiveComponents;
+	Algo::Transform(InStaticMeshComponentsToMerge, PrimitiveComponents, [](UStaticMeshComponent* SMComponent) { return SMComponent; });
+
+	FMaterialProxySettings MaterialProxySettings = InSettings.MaterialSettings;
+	if (MaterialProxySettings.ResolveTexelDensity(PrimitiveComponents))
+	{
+		MaterialProxySettings.TextureSize = InDataTracker.GetTextureSizeFromTargetTexelDensity(MaterialProxySettings.TargetTexelDensityPerMeter);
+		MaterialProxySettings.TextureSizingType = ETextureSizingType::TextureSizingType_UseSingleTextureSize;
+	}
+
+	UMaterialOptions* MaterialOptions = PopulateMaterialOptions(MaterialProxySettings);
+
+	// Check each material to see if the shader actually uses vertex data and collect flags
+	TArray<TOptional<bool>> bMaterialUsesVertexData;
+	bMaterialUsesVertexData.SetNum(InUniqueMaterials.Num());
+
+	// Deferred call, as this may not be required by all code paths and is pretty costly to compute
+	auto DoesMaterialUsesVertexData = [&](const int32 InMaterialIndex)
+	{
+		if (!bMaterialUsesVertexData[InMaterialIndex].IsSet())
 		{
-			//Make sure we have unique slot name here
-			FName MaterialSlotName = MergedMaterial->GetFName();
-			int32 Counter = 1;
-			while (!IsMaterialImportedNameUnique(MaterialSlotName))
+			bMaterialUsesVertexData[InMaterialIndex] = DetermineMaterialVertexDataUsage(InUniqueMaterials[InMaterialIndex], MaterialOptions);
+		}
+
+		return bMaterialUsesVertexData[InMaterialIndex].GetValue();
+	};
+
+	// For each unique material calculate how 'important' they are
+	TArray<float> MaterialImportanceValues;
+	FMaterialUtilities::DetermineMaterialImportance(InUniqueMaterials, MaterialImportanceValues);
+
+	TArray<FMeshData> GlobalMeshSettings;
+	TArray<FMaterialData> GlobalMaterialSettings;
+	TArray<float> SectionMaterialImportanceValues;
+
+	TMap<EMaterialProperty, FIntPoint> PropertySizes;
+	for (const FPropertyEntry& Entry : MaterialOptions->Properties)
+	{
+		if (!Entry.bUseConstantValue && Entry.Property != MP_MAX)
+		{
+			PropertySizes.Add(Entry.Property, Entry.bUseCustomSize ? Entry.CustomSize : MaterialOptions->TextureSize);
+		}
+	}
+
+	// If we are generating a single LOD and want to merge materials we can utilize texture space better by generating unique UVs
+	// for the merged mesh and baking out materials using those UVs
+	const bool bGloballyRemapUVs = !bInMergeAllLODs && !InSettings.bReuseMeshLightmapUVs;
+
+	typedef TTuple<UStaticMesh*, int32> FMeshLODTuple;
+	typedef TFuture<TArray<FVector2D>> FUVComputeFuture;
+	TMap<FMeshLODTuple, FUVComputeFuture> MeshLODsTextureCoordinates;
+	TMap<int32, FMeshLODTuple> MeshDataAwaitingResults;
+
+	for (TConstRawMeshIterator RawMeshIterator = InDataTracker.GetConstRawMeshIterator(); RawMeshIterator; ++RawMeshIterator)
+	{
+		const FMeshLODKey& Key = RawMeshIterator.Key();
+		const FMeshDescription& RawMesh = RawMeshIterator.Value();
+		const bool bRequiresUniqueUVs = InDataTracker.DoesMeshLODRequireUniqueUVs(Key);
+
+		const FMeshDescription* MeshDescription = InDataTracker.GetRawMeshPtr(Key);
+		UStaticMeshComponent* Component = InStaticMeshComponentsToMerge[Key.GetMeshIndex()];
+		UStaticMesh* StaticMesh = Component->GetStaticMesh();
+
+		// Retrieve all sections and materials for key
+		TArray<SectionRemapPair> SectionRemapPairs;
+		InDataTracker.GetMappingsForMeshLOD(Key, SectionRemapPairs);
+
+		// Contains unique materials used for this key, and the accompanying section index which point to the material
+		TMap<UMaterialInterface*, TArray<int32>> MaterialAndSectionIndices;
+
+		for (const SectionRemapPair& RemapPair : SectionRemapPairs)
+		{
+			const int32 UniqueIndex = RemapPair.Value;
+			const int32 SectionIndex = RemapPair.Key;
+			TArray<int32>& SectionIndices = MaterialAndSectionIndices.FindOrAdd(InCollapsedMaterialMap.FindChecked(InDataTracker.GetMaterialForSectionIndex(UniqueIndex)));
+			SectionIndices.Add(SectionIndex);
+		}
+
+		for (TPair<UMaterialInterface*, TArray<int32>>& MaterialSectionIndexPair : MaterialAndSectionIndices)
+		{
+			UMaterialInterface* Material = MaterialSectionIndexPair.Key;
+			const int32 MaterialIndex = InUniqueMaterials.IndexOfByKey(Material);
+			const TArray<int32>& SectionIndices = MaterialSectionIndexPair.Value;
+
+			FMaterialData MaterialData;
+			MaterialData.Material = InCollapsedMaterialMap.FindChecked(Material);
+			MaterialData.PropertySizes = PropertySizes;
+
+			FMeshData NewMeshData;
+			const bool bUseMeshData = bGloballyRemapUVs || (InSettings.bUseVertexDataForBakingMaterial && (bRequiresUniqueUVs || DoesMaterialUsesVertexData(MaterialIndex)));
+			if (bUseMeshData)
+			{
+				NewMeshData.Mesh = Key.GetMesh();
+				NewMeshData.MeshDescription = MeshDescription;
+				NewMeshData.VertexColorHash = Key.GetVertexColorHash();
+				NewMeshData.bMirrored = Component->GetComponentTransform().GetDeterminant() < 0.0f;
+				NewMeshData.MaterialIndices = SectionIndices;
+			}
+
+			auto CompareMaterialData = [&InSettings](const FMaterialData& LHS, const FMaterialData& RHS)
+			{
+				return InSettings.bMergeEquivalentMaterials ? FMaterialKey(LHS.Material) == FMaterialKey(RHS.Material) : LHS.Material == RHS.Material;
+			};
+
+			auto CompareMeshData = [](const FMeshData& LHS, const FMeshData& RHS)
 			{
-				MaterialSlotName = *(MergedMaterial->GetName() + TEXT("_") + FString::FromInt(Counter++));
+				return (LHS.Mesh == RHS.Mesh) && (LHS.MaterialIndices == RHS.MaterialIndices) && (LHS.bMirrored == RHS.bMirrored) && (LHS.VertexColorHash == RHS.VertexColorHash);
+			};
+
+			// Find material & mesh pair
+			int32 MeshDataIndex = INDEX_NONE;
+			for (int32 GlobalMaterialSettingsIndex = 0; GlobalMaterialSettingsIndex < GlobalMaterialSettings.Num(); ++GlobalMaterialSettingsIndex)
+			{
+				if (CompareMaterialData(GlobalMaterialSettings[GlobalMaterialSettingsIndex], MaterialData) &&
+					CompareMeshData(GlobalMeshSettings[GlobalMaterialSettingsIndex], NewMeshData))
+				{
+					MeshDataIndex = GlobalMaterialSettingsIndex;
+					break;
+				}
+			}
+
+			// We've found a match, no need to process this mesh/material pair
+			if (MeshDataIndex == INDEX_NONE)
+			{
+				// We're processing a new pair
+				MeshDataIndex = GlobalMeshSettings.Num();
+
+				FMeshData& MeshData = GlobalMeshSettings.Emplace_GetRef(NewMeshData);
+				GlobalMaterialSettings.Add(MaterialData);
+				SectionMaterialImportanceValues.Add(MaterialImportanceValues[MaterialIndex]);
+
+				if (bUseMeshData)
+				{
+					// if it has vertex color/*WedgetColors.Num()*/, it should also use light map UV index
+					// we can't do this for all meshes, but only for the mesh that has vertex color.
+					if (bRequiresUniqueUVs || MeshData.MeshDescription->VertexInstances().Num() > 0)
+					{
+						// Check if there are lightmap uvs available?
+						const int32 LightMapUVIndex = StaticMesh->GetLightMapCoordinateIndex();
+
+						TVertexInstanceAttributesConstRef<FVector2f> VertexInstanceUVs = FStaticMeshConstAttributes(*MeshData.MeshDescription).GetVertexInstanceUVs();
+						if (InSettings.bReuseMeshLightmapUVs && VertexInstanceUVs.GetNumElements() > 0 && VertexInstanceUVs.GetNumChannels() > LightMapUVIndex)
+						{
+							MeshData.TextureCoordinateIndex = LightMapUVIndex;
+						}
+						else
+						{
+							// Verify if we started an async task to generate UVs for this static mesh & LOD
+							FMeshLODTuple Tuple(Key.GetMesh(), Key.GetLODIndex());
+							if (!MeshLODsTextureCoordinates.Find(Tuple))
+							{
+								// No job found yet, fire an async task
+								MeshLODsTextureCoordinates.Add(Tuple, Async(EAsyncExecution::Thread, [MeshDescription, MaterialOptions, this]()
+									{
+										TArray<FVector2D> UniqueTextureCoordinates;
+										FStaticMeshOperations::GenerateUniqueUVsForStaticMesh(*MeshDescription, MaterialOptions->TextureSize.GetMax(), false, UniqueTextureCoordinates);
+										ScaleTextureCoordinatesToBox(FBox2D(FVector2D::ZeroVector, FVector2D(1, 1)), UniqueTextureCoordinates);
+										return UniqueTextureCoordinates;
+									}));
+							}
+							// Keep track of the fact that this mesh is waiting for the UV computation to finish
+							MeshDataAwaitingResults.Add(MeshDataIndex, Tuple);
+						}
+					}
+
+					InAdapters[Key.GetMeshIndex()].ApplySettings(Key.GetLODIndex(), MeshData);
+				}
+			}
+
+			for (const auto& OriginalSectionIndex : SectionIndices)
+			{
+				InOutputMaterialsMap.Add(Key, MaterialRemapPair(OriginalSectionIndex, MeshDataIndex));
 			}
-			StaticMesh->GetStaticMaterials().Add(FStaticMaterial(MergedMaterial, MaterialSlotName));
 		}
+	}
 
-		StaticMesh->PostEditChange();
+	// Fetch results from the async UV computation tasks
+	for (auto MeshData : MeshDataAwaitingResults)
+	{
+		GlobalMeshSettings[MeshData.Key].CustomTextureCoordinates = MeshLODsTextureCoordinates[MeshData.Value].Get();
+	}
 
-		OutAssetsToSync.Add(StaticMesh);
-		OutMergedActorLocation = MergedAssetPivot;
+	TArray<FMeshData*> MeshSettingPtrs;
+	for (int32 SettingsIndex = 0; SettingsIndex < GlobalMeshSettings.Num(); ++SettingsIndex)
+	{
+		MeshSettingPtrs.Add(&GlobalMeshSettings[SettingsIndex]);
 	}
+
+	TArray<FMaterialData*> MaterialSettingPtrs;
+	for (int32 SettingsIndex = 0; SettingsIndex < GlobalMaterialSettings.Num(); ++SettingsIndex)
+	{
+		MaterialSettingPtrs.Add(&GlobalMaterialSettings[SettingsIndex]);
+	}
+
+	if (bGloballyRemapUVs)
+	{
+		// We must keep vertex data in order to properly generate unique UVs
+		FMeshMergingSettings RemapUVMergeSettings = InSettings;
+		RemapUVMergeSettings.bBakeVertexDataToMesh = true;
+
+		TArray<FMeshDescription> MergedRawMeshes;
+		CreateMergedRawMeshes(InDataTracker, RemapUVMergeSettings, InStaticMeshComponentsToMerge, InUniqueMaterials, InCollapsedMaterialMap, InOutputMaterialsMap, false, false, InMergedAssetPivot, MergedRawMeshes);
+
+		// Create texture coords for the merged mesh
+		TArray<FVector2D> GlobalTextureCoordinates;
+		bool bSuccess = FStaticMeshOperations::GenerateUniqueUVsForStaticMesh(MergedRawMeshes[0], MaterialOptions->TextureSize.GetMax(), true, GlobalTextureCoordinates);
+		if (bSuccess)
+		{
+			ScaleTextureCoordinatesToBox(FBox2D(FVector2D::ZeroVector, FVector2D(1, 1)), GlobalTextureCoordinates);
+
+			// copy UVs back to the un-merged mesh's custom texture coords
+			// iterate the raw meshes in the same way as when we combined the mesh above in CreateMergedRawMeshes()
+			int32 GlobalUVIndex = 0;
+			for (TConstRawMeshIterator RawMeshIterator = InDataTracker.GetConstRawMeshIterator(); RawMeshIterator; ++RawMeshIterator)
+			{
+				const FMeshLODKey& Key = RawMeshIterator.Key();
+				const FMeshDescription& RawMesh = RawMeshIterator.Value();
+
+				// Build a local array for this raw mesh
+				TArray<FVector2D> UniqueTextureCoordinates;
+				UniqueTextureCoordinates.SetNumUninitialized(RawMesh.VertexInstances().Num());
+				for (FVector2D& UniqueTextureCoordinate : UniqueTextureCoordinates)
+				{
+					UniqueTextureCoordinate = GlobalTextureCoordinates[GlobalUVIndex++];
+				}
+
+				// copy to mesh data
+				for (FMeshData& MeshData : GlobalMeshSettings)
+				{
+					if (MeshData.MeshDescription == &RawMesh)
+					{
+						MeshData.CustomTextureCoordinates = UniqueTextureCoordinates;
+					}
+				}
+			}
+
+			// Dont smear borders as we will copy back non-pink pixels
+			for (FMaterialData& MaterialData : GlobalMaterialSettings)
+			{
+				MaterialData.bPerformBorderSmear = false;
+			}
+		}
+		else
+		{
+			UE_LOG(LogMeshMerging, Warning, TEXT("GenerateUniqueUVsForStaticMesh: Failed to pack UVs for static mesh"));
+		}
+	}
+
+	TArray<FFlattenMaterial> FlattenedMaterials;
+	// This scope ensures BakeOutputs is never used after TransferOutputToFlatMaterials
+	{
+		TArray<FBakeOutput> BakeOutputs;
+		IMaterialBakingModule& Module = FModuleManager::Get().LoadModuleChecked<IMaterialBakingModule>("MaterialBaking");
+		Module.BakeMaterials(MaterialSettingPtrs, MeshSettingPtrs, BakeOutputs);
+
+		// Append constant properties ?
+		TArray<FColor> ConstantData;
+		FIntPoint ConstantSize(1, 1);
+		for (const FPropertyEntry& Entry : MaterialOptions->Properties)
+		{
+			if (Entry.bUseConstantValue && Entry.Property != MP_MAX)
+			{
+				ConstantData.SetNum(1, false);
+				ConstantData[0] = FLinearColor(Entry.ConstantValue, Entry.ConstantValue, Entry.ConstantValue).ToFColor(true);
+				for (FBakeOutput& Output : BakeOutputs)
+				{
+					Output.PropertyData.Add(Entry.Property, ConstantData);
+					Output.PropertySizes.Add(Entry.Property, ConstantSize);
+				}
+			}
+		}
+
+		TransferOutputToFlatMaterials(GlobalMaterialSettings, BakeOutputs, FlattenedMaterials);
+	}
+
+	if (!bGloballyRemapUVs)
+	{
+		// Try to optimize materials where possible	
+		for (FFlattenMaterial& InMaterial : FlattenedMaterials)
+		{
+			FMaterialUtilities::OptimizeFlattenMaterial(InMaterial);
+		}
+	}
+
+	for (const FPropertyEntry& Entry : MaterialOptions->Properties)
+	{
+		if (Entry.Property != MP_MAX)
+		{
+			EFlattenMaterialProperties OldProperty = ToFlattenProperty(Entry.Property);
+			if (ensure(OldProperty != EFlattenMaterialProperties::NumFlattenMaterialProperties))
+			{
+				OutFlattenMaterial.SetPropertySize(OldProperty, Entry.bUseCustomSize ? Entry.CustomSize : MaterialOptions->TextureSize);
+			}
+		}
+	}
+
+	TArray<FUVOffsetScalePair> UVTransforms;
+	if (bGloballyRemapUVs)
+	{
+		// If we have globally remapped UVs we copy non-pink pixels over the dest texture rather than 
+		// copying sub-charts
+		TArray<FBox2D> MaterialBoxes;
+		MaterialBoxes.SetNumUninitialized(GlobalMaterialSettings.Num());
+		for (FBox2D& Box2D : MaterialBoxes)
+		{
+			Box2D = FBox2D(FVector2D(0.0f, 0.0f), FVector2D(1.0f, 1.0f));
+		}
+
+		FlattenBinnedMaterials(FlattenedMaterials, MaterialBoxes, 0, true, OutFlattenMaterial, UVTransforms);
+	}
+	else
+	{
+		/** Reweighting */
+		float TotalValue = 0.0f;
+		for (const float& Value : SectionMaterialImportanceValues)
+		{
+			TotalValue += Value;
+		}
+
+		float Multiplier = 1.0f / TotalValue;
+
+		for (float& Value : SectionMaterialImportanceValues)
+		{
+			Value *= Multiplier;
+		}
+		/** End reweighting */
+
+		if (InSettings.bUseTextureBinning)
+		{
+			TArray<FBox2D> MaterialBoxes;
+			FMaterialUtilities::GeneratedBinnedTextureSquares(FVector2D(1.0f, 1.0f), SectionMaterialImportanceValues, MaterialBoxes);
+			FlattenBinnedMaterials(FlattenedMaterials, MaterialBoxes, InSettings.GutterSize, false, OutFlattenMaterial, UVTransforms);
+		}
+		else
+		{
+			MergeFlattenedMaterials(FlattenedMaterials, InSettings.GutterSize, OutFlattenMaterial, UVTransforms);
+		}
+	}
+
+	// Adjust UVs
+	for (int32 ComponentIndex = 0; ComponentIndex < InStaticMeshComponentsToMerge.Num(); ++ComponentIndex)
+	{
+		TArray<uint32> ProcessedMaterials;
+		for (TPair<FMeshLODKey, MaterialRemapPair>& MappingPair : InOutputMaterialsMap)
+		{
+			if (MappingPair.Key.GetMeshIndex() == ComponentIndex && !ProcessedMaterials.Contains(MappingPair.Value.Key))
+			{
+				// Retrieve raw mesh data for this component and lod pair
+				FMeshDescription* RawMesh = InDataTracker.GetRawMeshPtr(MappingPair.Key);
+
+				FMeshData& MeshData = GlobalMeshSettings[MappingPair.Value.Value];
+				const FUVOffsetScalePair& UVTransform = UVTransforms[MappingPair.Value.Value];
+
+				const uint32 MaterialIndex = MappingPair.Value.Key;
+				ProcessedMaterials.Add(MaterialIndex);
+				if (RawMesh->Vertices().Num())
+				{
+					TVertexInstanceAttributesRef<FVector2f> VertexInstanceUVs = FStaticMeshAttributes(*RawMesh).GetVertexInstanceUVs();
+					int32 NumUVChannel = FMath::Min(VertexInstanceUVs.GetNumChannels(), (int32)MAX_MESH_TEXTURE_COORDS);
+					for (int32 UVChannelIdx = 0; UVChannelIdx < NumUVChannel; ++UVChannelIdx)
+					{
+						int32 VertexIndex = 0;
+						for (FVertexInstanceID VertexInstanceID : RawMesh->VertexInstances().GetElementIDs())
+						{
+							FVector2D UV = FVector2D(VertexInstanceUVs.Get(VertexInstanceID, UVChannelIdx));
+							if (UVChannelIdx == 0)
+							{
+								if (MeshData.CustomTextureCoordinates.Num())
+								{
+									UV = MeshData.CustomTextureCoordinates[VertexIndex];
+								}
+								else if (MeshData.TextureCoordinateIndex != 0)
+								{
+									check(MeshData.TextureCoordinateIndex < NumUVChannel);
+									UV = FVector2D(VertexInstanceUVs.Get(VertexInstanceID, MeshData.TextureCoordinateIndex));
+								}
+							}
+
+							const TArray<FPolygonID>& Polygons = RawMesh->GetVertexInstanceConnectedPolygons(VertexInstanceID);
+							for (FPolygonID PolygonID : Polygons)
+							{
+								FPolygonGroupID PolygonGroupID = RawMesh->GetPolygonPolygonGroup(PolygonID);
+								if (PolygonGroupID.GetValue() == MaterialIndex)
+								{
+									if (UVTransform.Value != FVector2D::ZeroVector)
+									{
+										VertexInstanceUVs.Set(VertexInstanceID, UVChannelIdx, FVector2f(UV * UVTransform.Value + UVTransform.Key));	// LWC_TODO: Precision loss
+										break;
+									}
+								}
+							}
+							VertexIndex++;
+						}
+					}
+				}
+			}
+		}
+	}
+
+	for (TRawMeshIterator Iterator = InDataTracker.GetRawMeshIterator(); Iterator; ++Iterator)
+	{
+		FMeshDescription& RawMesh = Iterator.Value();
+		// Reset material indexes
+		TMap<FPolygonGroupID, FPolygonGroupID> RemapPolygonGroups;
+		for (FPolygonGroupID PolygonGroupID : RawMesh.PolygonGroups().GetElementIDs())
+		{
+			RemapPolygonGroups.Add(PolygonGroupID, FPolygonGroupID(0));
+		}
+		RawMesh.RemapPolygonGroups(RemapPolygonGroups);
+	}
+
+	OutFlattenMaterial.UVChannel = INDEX_NONE;
 }
 
 void FMeshMergeUtilities::CreateMergedRawMeshes(FMeshMergeDataTracker& InDataTracker, const FMeshMergingSettings& InSettings, const TArray<UStaticMeshComponent*>& InStaticMeshComponentsToMerge, const TArray<UMaterialInterface*>& InUniqueMaterials, const TMap<UMaterialInterface*, UMaterialInterface*>& InCollapsedMaterialMap, const TMultiMap<FMeshLODKey, MaterialRemapPair>& InOutputMaterialsMap, bool bInMergeAllLODs, bool bInMergeMaterialData, const FVector& InMergedAssetPivot, TArray<FMeshDescription>& OutMergedRawMeshes) const
@@ -3190,7 +3142,6 @@ void FMeshMergeUtilities::CreateMergedRawMeshes(FMeshMergeDataTracker& InDataTra
 				int32 RetrievedLODIndex = LODIndex;
 				FMeshDescription* RawMeshPtr = InDataTracker.TryFindRawMeshForLOD(ComponentIndex, RetrievedLODIndex);
 
-
 				if (RawMeshPtr != nullptr)
 				{
 					InDataTracker.AddComponentToWedgeMapping(ComponentIndex, LODIndex, MergedMesh.VertexInstances().Num());
@@ -3665,7 +3616,7 @@ void FMeshMergeUtilities::MergeComponentsToInstances(const TArray<UPrimitiveComp
 	}
 }
 
-UMaterialInterface* FMeshMergeUtilities::CreateProxyMaterial(const FString &InBasePackageName, FString MergedAssetPackageName, UMaterialInterface* InBaseMaterial, UPackage* InOuter, const FMeshMergingSettings &InSettings, FFlattenMaterial OutMaterial, TArray<UObject *>& OutAssetsToSync, FMaterialUpdateContext* InMaterialUpdateContext) const
+UMaterialInterface* FMeshMergeUtilities::CreateProxyMaterial(const FString &InBasePackageName, FString MergedAssetPackageName, UMaterialInterface* InBaseMaterial, UPackage* InOuter, const FMeshMergingSettings &InSettings, const FFlattenMaterial& OutMaterial, TArray<UObject *>& OutAssetsToSync, FMaterialUpdateContext* InMaterialUpdateContext) const
 {
 	// Create merged material asset
 	FString MaterialAssetName;
diff --git a/Engine/Source/Developer/MeshMergeUtilities/Private/MeshMergeUtilities.h b/Engine/Source/Developer/MeshMergeUtilities/Private/MeshMergeUtilities.h
--- a/Engine/Source/Developer/MeshMergeUtilities/Private/MeshMergeUtilities.h
+++ b/Engine/Source/Developer/MeshMergeUtilities/Private/MeshMergeUtilities.h
@@ -4,6 +4,7 @@
 
 #include "IMeshMergeUtilities.h"
 #include "SceneTypes.h"
+#include "StaticMeshComponentAdapter.h"
 #include "MaterialShared.h"
 
 class FMeshMergeDataTracker;
@@ -15,7 +16,7 @@ typedef TPair<uint32, uint32> MaterialRemapPair;
 /**
  * Mesh Merge Utilities
  */
-class MESHMERGEUTILITIES_API FMeshMergeUtilities : public IMeshMergeUtilities
+class FMeshMergeUtilities : public IMeshMergeUtilities
 {
 public:	
 	friend class FProxyGenerationProcessor;
@@ -30,10 +31,7 @@ public:
 
 	virtual void CreateProxyMesh(const TArray<AActor*>& InActors, const struct FMeshProxySettings& InMeshProxySettings, UMaterialInterface* InBaseMaterial, UPackage* InOuter, const FString& InProxyBasePackageName, const FGuid InGuid, const FCreateProxyDelegate& InProxyCreatedDelegate, const bool bAllowAsync = false, const float ScreenSize = 1.0f) const override;
 	virtual void CreateProxyMesh(const TArray<UStaticMeshComponent*>& InStaticMeshComps, const struct FMeshProxySettings& InMeshProxySettings, UMaterialInterface* InBaseMaterial, UPackage* InOuter, const FString& InProxyBasePackageName, const FGuid InGuid, const FCreateProxyDelegate& InProxyCreatedDelegate, const bool bAllowAsync = false, const float ScreenSize = 1.0f) const override;
-
-	UE_DEPRECATED(4.20, "Use different signature containing BaseMaterial parameter")
 	virtual void CreateProxyMesh(const TArray<AActor*>& InActors, const struct FMeshProxySettings& InMeshProxySettings, UPackage* InOuter, const FString& InProxyBasePackageName, const FGuid InGuid, const FCreateProxyDelegate& InProxyCreatedDelegate, const bool bAllowAsync = false, const float ScreenSize = 1.0f) const override;
-	UE_DEPRECATED(4.20, "Use different signature containing BaseMaterial parameter")
 	virtual void CreateProxyMesh(const TArray<UStaticMeshComponent*>& InStaticMeshComps, const struct FMeshProxySettings& InMeshProxySettings, UPackage* InOuter, const FString& InProxyBasePackageName, const FGuid InGuid, const FCreateProxyDelegate& InProxyCreatedDelegate, const bool bAllowAsync = false, const float ScreenSize = 1.0f) const override;
 
 	UE_DEPRECATED(5.0, "Use FMaterialUtilities::IsValidFlattenMaterial()")
@@ -49,11 +47,14 @@ public:
 	virtual void UnregisterExtension(IMeshMergeExtension* InExtension) override;	
 protected:
 	/** Creates a proxy material instance at givne path and name */
-	UMaterialInterface* CreateProxyMaterial(const FString &InBasePackageName, FString MergedAssetPackageName, UMaterialInterface* InBaseMaterial, UPackage* InOuter, const FMeshMergingSettings &InSettings, FFlattenMaterial OutMaterial, TArray<UObject *>& OutAssetsToSync, FMaterialUpdateContext* InMaterialUpdateContext = nullptr) const;
+	UMaterialInterface* CreateProxyMaterial(const FString &InBasePackageName, FString MergedAssetPackageName, UMaterialInterface* InBaseMaterial, UPackage* InOuter, const FMeshMergingSettings &InSettings, const FFlattenMaterial& OutMaterial, TArray<UObject *>& OutAssetsToSync, FMaterialUpdateContext* InMaterialUpdateContext = nullptr) const;
 	/** Merges flattened material into atlas textures */
 	void MergeFlattenedMaterials(TArray<struct FFlattenMaterial>& InMaterialList, int32 InGutter, FFlattenMaterial& OutMergedMaterial, TArray<FUVOffsetScalePair>& OutUVTransforms) const;
 	/** Merges flattened material into binned textures */
 	void FlattenBinnedMaterials(TArray<struct FFlattenMaterial>& InMaterialList, const TArray<FBox2D>& InMaterialBoxes, int32 InGutter, bool bCopyOnlyMaskedPixels, FFlattenMaterial& OutMergedMaterial, TArray<FUVOffsetScalePair>& OutUVTransforms) const;
+
+	void CreateMergedMaterial(FMeshMergeDataTracker& InDataTracker, const FMeshMergingSettings& InSettings, const TArray<UStaticMeshComponent*>& InStaticMeshComponentsToMerge, TArray<FStaticMeshComponentAdapter>& InAdapters, const TArray<UMaterialInterface*>& InUniqueMaterials, const TMap<UMaterialInterface*, UMaterialInterface*>& InCollapsedMaterialMap, TMultiMap<FMeshLODKey, MaterialRemapPair>& InOutputMaterialsMap, bool bInMergeAllLODs, bool bInMergeMaterialData, const FVector& InMergedAssetPivot, FFlattenMaterial& OutFlattenMaterial) const;
+
 	/** Helper function to create the final merged raw meshes */
 	void CreateMergedRawMeshes(FMeshMergeDataTracker& InDataTracker, const FMeshMergingSettings& InSettings, const TArray<UStaticMeshComponent*>& InStaticMeshComponentsToMerge, const TArray<UMaterialInterface*>& InUniqueMaterials, const TMap<UMaterialInterface*, UMaterialInterface*>& InCollapsedMaterialMap, const TMultiMap<FMeshLODKey, MaterialRemapPair>& InOutputMaterialsMap, bool bInMergeAllLODs, bool bInMergeMaterialData, const FVector& InMergedAssetPivot, TArray<FMeshDescription>& OutMergedRawMeshes) const;
 protected:
@@ -70,7 +71,6 @@ protected:
 	/** Conditionally resizes the source data into OutImage */
 	FIntPoint ConditionalImageResize(const FIntPoint& SrcSize, const FIntPoint& DesiredSize, TArray<FColor>& InOutImage, bool bLinearSpace) const;
 	/** Converts bake output structure data to flatten material format */
-	UE_DEPRECATED(4.25, "Use TransferOutputToFlatMaterials instead for improved efficiency")
 	void ConvertOutputToFlatMaterials(const TArray<FBakeOutput>& BakeOutputs, const TArray<FMaterialData>& MaterialData, TArray<FFlattenMaterial> &FlattenedMaterials) const;
 	/** Transfer bake output structure data to flatten material format.
 	  * @warning This is a destructive operation for InOutBakeOutputs
diff --git a/Engine/Source/Runtime/Engine/Classes/Engine/MeshMerging.h b/Engine/Source/Runtime/Engine/Classes/Engine/MeshMerging.h
--- a/Engine/Source/Runtime/Engine/Classes/Engine/MeshMerging.h
+++ b/Engine/Source/Runtime/Engine/Classes/Engine/MeshMerging.h
@@ -530,10 +530,6 @@ struct FMeshMergingSettings
 	UPROPERTY(EditAnywhere, BlueprintReadWrite, Category = MaterialSettings, meta=(EditCondition="LODSelectionType == EMeshLODSelectionType::LowestDetailLOD"))
 	uint8 bMergeMaterials:1;
 
-	/** Create a flat material from all source materials, along with a new set of UVs. This material won't be applied to any section by default. */
-	UPROPERTY(EditAnywhere, AdvancedDisplay, BlueprintReadWrite, Category = MaterialSettings, meta=(EditCondition="LODSelectionType == EMeshLODSelectionType::LowestDetailLOD"))
-	uint8 bCreateMergedMaterial : 1;
-
 	/** Whether or not vertex data such as vertex colours should be baked into the resulting mesh */
 	UPROPERTY(EditAnywhere, BlueprintReadWrite, Category = MeshSettings)
 	uint8 bBakeVertexDataToMesh:1;
@@ -593,6 +589,9 @@ struct FMeshMergingSettings
 	UPROPERTY()
 	uint8 bExportSpecularMap_DEPRECATED:1;
 
+	UPROPERTY()
+	uint8 bCreateMergedMaterial_DEPRECATED : 1;
+
 	UPROPERTY()
 	int32 MergedMaterialAtlasResolution_DEPRECATED;
 
@@ -603,7 +602,7 @@ struct FMeshMergingSettings
 	uint8 bGenerateNaniteEnabledMesh_DEPRECATED : 1;
 
 	UPROPERTY()
-	float NaniteFallbackTrianglePercent_DEPRECATED;
+	float NaniteFallbackTrianglePercent_DEPRECATED;	
 #endif
 
 	EMeshMergeType MergeType;
@@ -620,7 +619,6 @@ struct FMeshMergingSettings
 		, bMergePhysicsData(false)
 		, bMergeMeshSockets(false)
 		, bMergeMaterials(false)
-		, bCreateMergedMaterial(false)
 		, bBakeVertexDataToMesh(false)
 		, bUseVertexDataForBakingMaterial(true)
 		, bUseTextureBinning(false)
@@ -637,6 +635,7 @@ struct FMeshMergingSettings
 		, bExportMetallicMap_DEPRECATED(false)
 		, bExportRoughnessMap_DEPRECATED(false)
 		, bExportSpecularMap_DEPRECATED(false)
+		, bCreateMergedMaterial_DEPRECATED(false)
 		, MergedMaterialAtlasResolution_DEPRECATED(1024)
 		, ExportSpecificLOD_DEPRECATED(0)
 		, bGenerateNaniteEnabledMesh_DEPRECATED(false)
