FHierarchicalLODUtilities::BuildStaticMeshForLODActor() - When adding impostors, should use the component transform rather than the actor transform

PR #8624: HLOD imposters take component instead of actor transform (Contributed by [REDACTED_PERSON])

diff --git a/Engine/Source/Developer/HierarchicalLODUtilities/Private/HierarchicalLODUtilities.cpp b/Engine/Source/Developer/HierarchicalLODUtilities/Private/HierarchicalLODUtilities.cpp
--- a/Engine/Source/Developer/HierarchicalLODUtilities/Private/HierarchicalLODUtilities.cpp
+++ b/Engine/Source/Developer/HierarchicalLODUtilities/Private/HierarchicalLODUtilities.cpp
@@ -837,21 +837,21 @@ bool FHierarchicalLODUtilities::BuildStaticMeshForLODActor(ALODActor* LODActor,
 
 			FLODInstanceBatch& LODInstanceBatch = InstancesBatches.FindOrAdd(Key);
 
+			FTransform ComponentTransformWS = SMC->GetComponentTransform();
+
 			// If we have an ISMC, ensure we include all its instances
 			if (UInstancedStaticMeshComponent* InstancedSMC = Cast<UInstancedStaticMeshComponent>(SMC))
 			{
-				FTransform ActorTransformWS = InstancedSMC->GetOwner()->GetActorTransform();
-
 				LODInstanceBatch.Transforms.Reserve(LODInstanceBatch.Transforms.Num() + InstancedSMC->GetInstanceCount());
 				for (const FInstancedStaticMeshInstanceData& InstanceData : InstancedSMC->PerInstanceSMData)
 				{
-					FTransform InstanceTransformWS = FTransform(InstanceData.Transform) * ActorTransformWS;
+					FTransform InstanceTransformWS = FTransform(InstanceData.Transform) * ComponentTransformWS;
 					LODInstanceBatch.Transforms.Add(InstanceTransformWS);
 				}
 			}
 			else
 			{
-				LODInstanceBatch.Transforms.Add(SMC->GetOwner()->GetActorTransform());
+				LODInstanceBatch.Transforms.Add(ComponentTransformWS);
 				LODInstanceBatch.CustomPrimitiveData.Add(SMC->GetCustomPrimitiveData());
 			}
 
diff --git a/Engine/Source/Developer/MeshMergeUtilities/Private/MeshMergeHelpers.cpp b/Engine/Source/Developer/MeshMergeUtilities/Private/MeshMergeHelpers.cpp
--- a/Engine/Source/Developer/MeshMergeUtilities/Private/MeshMergeHelpers.cpp
+++ b/Engine/Source/Developer/MeshMergeUtilities/Private/MeshMergeHelpers.cpp
@@ -1162,19 +1162,19 @@ void FMeshMergeHelpers::MergeImpostersToMesh(TArray<const UStaticMeshComponent*>
 			TVertexInstanceAttributesRef<FVector2f> VertexInstanceUVs = ImposterMeshAttributes.GetVertexInstanceUVs();
 			VertexInstanceUVs.SetNumChannels(UVTwoIndex + 1);
 			const int32 NumIndices = ImposterMesh.VertexInstances().Num();
-			const FTransform& ActorToWorld = Component->GetOwner()->GetActorTransform();
-			const FVector ActorPosition = ActorToWorld.TransformPosition(FVector::ZeroVector) - InPivot;
+			const FTransform& ComponentToWorld = Component->GetComponentTransform();
+			const FVector ComponentPosition = ComponentToWorld.TransformPosition(FVector::ZeroVector) - InPivot;
 			for (const FVertexInstanceID VertexInstanceID : ImposterMesh.VertexInstances().GetElementIDs())
 			{
 				FVector2D UVOne;
 				FVector2D UVTwo;
 
-				UVOne.X = ActorPosition.X;
-				UVOne.Y = ActorPosition.Y;
+				UVOne.X = ComponentPosition.X;
+				UVOne.Y = ComponentPosition.Y;
 				VertexInstanceUVs.Set(VertexInstanceID, UVOneIndex, FVector2f(UVOne));	// LWC_TODO: Precision loss
 
-				UVTwo.X = ActorPosition.Z;
-				UVTwo.Y = FMath::Abs(ActorToWorld.GetScale3D().X);
+				UVTwo.X = ComponentPosition.Z;
+				UVTwo.Y = FMath::Abs(ComponentToWorld.GetScale3D().X);
 				VertexInstanceUVs.Set(VertexInstanceID, UVTwoIndex, FVector2f(UVTwo));	// LWC_TODO: Precision loss
 			}
 		}
