Mesh Merging - Fix combination of FMeshProxySettings settings bReuseMeshLightmapUVs and bGroupIdenticalMeshesForBaking

diff --git a/Engine/Source/Developer/MeshMergeUtilities/Private/MeshMergeUtilities.cpp b/Engine/Source/Developer/MeshMergeUtilities/Private/MeshMergeUtilities.cpp
--- a/Engine/Source/Developer/MeshMergeUtilities/Private/MeshMergeUtilities.cpp
+++ b/Engine/Source/Developer/MeshMergeUtilities/Private/MeshMergeUtilities.cpp
@@ -1089,6 +1089,37 @@ float FMeshMergeUtilities::FlattenEmissivescale(TArray<struct FFlattenMaterial>&
 	return MaxScale;
 }
 
+static TArray<FVector2D> GetCustomTextureCoordinates(const FMeshDescription& InMeshDescription, const UStaticMesh* InStaticMesh, const FMeshProxySettings& InMeshProxySettings)
+{
+	TArray<FVector2D> CustomTextureCoordinates;
+
+	TVertexInstanceAttributesConstRef<FVector2f> VertexInstanceUVs = FStaticMeshConstAttributes(InMeshDescription).GetVertexInstanceUVs();
+
+	// If we already have lightmap uvs generated and they are valid, we can reuse those instead of having to generate new ones
+	const int32 LightMapCoordinateIndex = InStaticMesh->GetLightMapCoordinateIndex();
+	if (InMeshProxySettings.bReuseMeshLightmapUVs &&
+		LightMapCoordinateIndex > 0 &&
+		VertexInstanceUVs.GetNumElements() > 0 &&
+		VertexInstanceUVs.GetNumChannels() > LightMapCoordinateIndex)
+	{
+		for (const FVertexInstanceID VertexInstanceID : InMeshDescription.VertexInstances().GetElementIDs())
+		{
+			CustomTextureCoordinates.Add(FVector2D(VertexInstanceUVs.Get(VertexInstanceID, LightMapCoordinateIndex)));
+		}
+	}
+	else
+	{
+		bool bSuccess = FStaticMeshOperations::GenerateUniqueUVsForStaticMesh(InMeshDescription, InMeshProxySettings.MaterialSettings.TextureSize.GetMax(), false, CustomTextureCoordinates);
+		if (!bSuccess)
+		{
+			UE_LOG(LogMeshMerging, Warning, TEXT("GenerateUniqueUVsForStaticMesh: Failed to pack UVs for static mesh \"%s\" (num triangles = %d, texture resolution = %d)."), *InStaticMesh->GetName(), InMeshDescription.Triangles().Num(), InMeshProxySettings.MaterialSettings.TextureSize.GetMax());
+			CustomTextureCoordinates.Empty();
+		}
+	}
+
+	return CustomTextureCoordinates;
+}
+
 class FProxyMeshDescriptor
 {
 public:
@@ -1138,11 +1169,18 @@ public:
 
 	UStaticMesh* GetStaticMesh() const { return ISMDescriptor.StaticMesh; }
 
+	bool IsValid() const { return !MeshDescription.IsEmpty(); }
+
 	const FMeshDescription& GetMeshDescription() const
 	{
 		return MeshDescription;
 	}
 
+	const TArray<FVector2D>& GetCustomTextureCoordinates() const
+	{
+		return CustomTextureCoordinates;
+	}
+
 	void PrepareMeshDescription(const FMeshProxySettings& InMeshProxySettings)
 	{
 		TRACE_CPUPROFILER_EVENT_SCOPE(PrepareMeshDescription);
@@ -1151,20 +1189,10 @@ public:
 		FStaticMeshAttributes(MeshDescription).Register();
 		FMeshMergeHelpers::RetrieveMesh(ISMDescriptor.StaticMesh, LODIndex, MeshDescription);
 
-		TArray<FVector2D> TexCoords;
-		bool bSuccess = FStaticMeshOperations::GenerateUniqueUVsForStaticMesh(MeshDescription, InMeshProxySettings.MaterialSettings.TextureSize.GetMax(), false, TexCoords);
-
-		if (bSuccess)
-		{
-			TVertexInstanceAttributesRef<FVector2f> VertexInstanceUVs = FStaticMeshAttributes(MeshDescription).GetVertexInstanceUVs();
-			for (int32 Idx = 0; Idx < TexCoords.Num(); Idx++)
-			{
-				VertexInstanceUVs.Set(Idx, FVector2f(TexCoords[Idx]));	// LWC_TODO: Precision loss
-			}
-		}
-		else
+		CustomTextureCoordinates = ::GetCustomTextureCoordinates(MeshDescription, ISMDescriptor.StaticMesh, InMeshProxySettings);
+		if (CustomTextureCoordinates.IsEmpty())
 		{
-			UE_LOG(LogMeshMerging, Warning, TEXT("GenerateUniqueUVsForStaticMesh: Failed to pack UVs for static mesh \"%s\" (num triangles = %d, texture resolution = %d)."), *ISMDescriptor.StaticMesh->GetName(), MeshDescription.Triangles().Num(), InMeshProxySettings.MaterialSettings.TextureSize.GetMax());
+			// Failure, clear mesh description
 			MeshDescription.Empty();
 		}
 	}
@@ -1179,6 +1207,7 @@ private:
 	FISMComponentDescriptor ISMDescriptor;
 
 	FMeshDescription MeshDescription;
+	TArray<FVector2D> CustomTextureCoordinates;
 };
 
 static void ScaleTextureCoordinatesToBox(const FBox2D& Box, TArray<FVector2D>& InOutTextureCoordinates)
@@ -1236,7 +1265,7 @@ static TArray<FInstancedMeshDescriptionData> GatherGeometry(const TArray<UStatic
 			FMeshDescription& MeshDescription = TempDescriptionData[Index];
 
 			// Retrieve meshes
-			if (!InMeshProxySettings.bGroupIdenticalMeshesForBaking || InDescriptors[InMeshesToMeshDescriptor[Index]].GetMeshDescription().IsEmpty())
+			if (!InMeshProxySettings.bGroupIdenticalMeshesForBaking || !InDescriptors[InMeshesToMeshDescriptor[Index]].IsValid())
 			{
 				const int32 LODIndex = InGetMeshLODFunc(StaticMeshComponent);
 				static const bool bPropagateVertexColours = true;
@@ -1301,7 +1330,6 @@ static TArray<FInstancedMeshDescriptionData> GatherGeometry(const TArray<UStatic
 					if (InMeshProxySettings.bGroupIdenticalMeshesForBaking)
 					{
 						FTransform ComponentTransform = StaticMeshComponent->GetComponentTransform();
-						FTransform ComponentTransformInv = ComponentTransform.Inverse();
 
 						if (const UInstancedStaticMeshComponent* InstancedStaticMeshComponent = Cast<UInstancedStaticMeshComponent>(StaticMeshComponent))
 						{
@@ -1354,82 +1382,38 @@ TArray<FMeshData> PrepareBakingMeshes(const struct FMeshProxySettings& InMeshPro
 	TArray<FMeshData> MeshData;
 	MeshData.SetNum(InDescriptors.Num());
 
-	// Step that must run on the game thread
-	for (int32 MeshIndex = 0; MeshIndex < InDescriptors.Num(); MeshIndex++)
-	{
-		FMeshData& MeshSettings = MeshData[MeshIndex];
-		FMeshDescription& MeshDescription = *InMeshDescriptionData[MeshIndex].MeshDescription;
-		const FProxyMeshDescriptor& MeshDescriptor = InDescriptors[MeshIndex];
-
-		if (!InMeshProxySettings.bGroupIdenticalMeshesForBaking)
-		{
-			if (MeshDescriptor.GetLightMapIndex() != INDEX_NONE)
-			{
-				MeshSettings.LightMap = MeshDescriptor.GetLightMap();
-				MeshSettings.LightMapIndex = MeshDescriptor.GetLightMapIndex();
-			}
-		}
-	}
-
 	// Parallel step
 	ParallelFor(InDescriptors.Num(), [&MeshData, &InDescriptors, &InMeshDescriptionData, &InMeshProxySettings](uint32 MeshIndex)
 	{
 		const FProxyMeshDescriptor& MeshDescriptor = InDescriptors[MeshIndex];
 
 		FMeshData& MeshSettings = MeshData[MeshIndex];
+		MeshSettings.TextureCoordinateBox = FBox2D(FVector2D(0.0f, 0.0f), FVector2D(1.0f, 1.0f));
 
-		bool bSuccess = true;
+
+		if (MeshDescriptor.GetLightMapIndex() != INDEX_NONE)
+		{
+			MeshSettings.LightMap = MeshDescriptor.GetLightMap();
+			MeshSettings.LightMapIndex = MeshDescriptor.GetLightMapIndex();
+		}
 
 		if (InMeshProxySettings.bGroupIdenticalMeshesForBaking)
 		{
 			// Grouping by identical meshes, the UVs should have already been setup
-			bSuccess = !MeshDescriptor.GetMeshDescription().IsEmpty();
 			MeshSettings.MeshDescription = &MeshDescriptor.GetMeshDescription();
+			MeshSettings.CustomTextureCoordinates = MeshDescriptor.GetCustomTextureCoordinates();
 		}
 		else
 		{
 			FMeshDescription& MeshDescription = *InMeshDescriptionData[MeshIndex].MeshDescription;
 			MeshSettings.MeshDescription = &MeshDescription;
-
-			TVertexInstanceAttributesRef<FVector2f> VertexInstanceUVs = FStaticMeshAttributes(MeshDescription).GetVertexInstanceUVs();
-
-		    // If we already have lightmap uvs generated and they are valid, we can reuse those instead of having to generate new ones
-		    const int32 LightMapCoordinateIndex = MeshDescriptor.GetStaticMesh()->GetLightMapCoordinateIndex();
-		    if (InMeshProxySettings.bReuseMeshLightmapUVs &&
-			    LightMapCoordinateIndex > 0 &&
-			    VertexInstanceUVs.GetNumElements() > 0 &&
-			    VertexInstanceUVs.GetNumChannels() > LightMapCoordinateIndex)
-		    {
-			    MeshSettings.CustomTextureCoordinates.Reset(VertexInstanceUVs.GetNumElements());
-			    for (const FVertexInstanceID VertexInstanceID : MeshDescription.VertexInstances().GetElementIDs())
-			    {
-				    MeshSettings.CustomTextureCoordinates.Add(FVector2D(VertexInstanceUVs.Get(VertexInstanceID, LightMapCoordinateIndex)));
-			    }
-			    MeshSettings.TextureCoordinateBox = FBox2D(FVector2D::ZeroVector, FVector2D(1, 1));
-			    ScaleTextureCoordinatesToBox(MeshSettings.TextureCoordinateBox, MeshSettings.CustomTextureCoordinates);
-		    }
-		    else
-		    {
-			    // Generate unique UVs for mesh (should only be done if needed)
-			    bSuccess = FStaticMeshOperations::GenerateUniqueUVsForStaticMesh(MeshDescription, InMeshProxySettings.MaterialSettings.TextureSize.GetMax(), false, MeshSettings.CustomTextureCoordinates);
-			    if (bSuccess)
-			    {
-				    MeshSettings.TextureCoordinateBox = FBox2D(FVector2D::ZeroVector, FVector2D(1, 1));
-				    ScaleTextureCoordinatesToBox(MeshSettings.TextureCoordinateBox, MeshSettings.CustomTextureCoordinates);
-			    }
-			    else
-			    {
-				    UE_LOG(LogMeshMerging, Warning, TEXT("GenerateUniqueUVsForStaticMesh: Failed to pack UVs for static mesh \"%s\" (num triangles = %d, texture resolution = %d)."), *MeshDescriptor.GetStaticMesh()->GetName(), MeshDescription.Triangles().Num(), InMeshProxySettings.MaterialSettings.TextureSize.GetMax());
-			    }
-		    }
+			MeshSettings.CustomTextureCoordinates = GetCustomTextureCoordinates(MeshDescription, MeshDescriptor.GetStaticMesh(), InMeshProxySettings);
 		}
 
-		if (!bSuccess)
+		if (MeshSettings.CustomTextureCoordinates.IsEmpty())
 		{
 			MeshSettings.MeshDescription = nullptr;
-			MeshSettings.TextureCoordinateBox = FBox2D(FVector2D(0.0f, 0.0f), FVector2D(1.0f, 1.0f));
 			MeshSettings.TextureCoordinateIndex = 0;
-			MeshSettings.CustomTextureCoordinates.Empty();
 		}
 	});
 
@@ -1883,32 +1867,36 @@ void FMeshMergeUtilities::CreateProxyMesh(const TArray<UStaticMeshComponent*>& I
 			FInstancedMeshMergeData MergeData;
 			MergeData.SourceStaticMesh = MeshDescriptors[Index].GetStaticMesh();
 			MergeData.RawMesh = MeshDescriptionData[Index].MeshDescription;
+			MergeData.NewUVs = MeshDescriptors[Index].GetCustomTextureCoordinates();
 			MergeData.bIsClippingMesh = false;
 			MergeData.InstanceTransforms = MeshDescriptionData[Index].InstancesTransforms;
 
 			FMeshMergeHelpers::CalculateTextureCoordinateBoundsForMesh(*MergeData.RawMesh, MergeData.TexCoordBounds);
 
-			FMeshData* MeshData = GlobalMeshSettings.FindByPredicate([&](const FMeshData& Entry)
-			{
-				return Entry.MeshDescription == MergeData.RawMesh && (Entry.CustomTextureCoordinates.Num() || Entry.TextureCoordinateIndex != 0);
-			});
-
-			if (MeshData)
+			if (MergeData.NewUVs.IsEmpty())
 			{
-				if (MeshData->CustomTextureCoordinates.Num())
+				FMeshData* MeshData = GlobalMeshSettings.FindByPredicate([&](const FMeshData& Entry)
 				{
-					MergeData.NewUVs = MeshData->CustomTextureCoordinates;
-				}
-				else
+					return Entry.MeshDescription == MergeData.RawMesh && (Entry.CustomTextureCoordinates.Num() || Entry.TextureCoordinateIndex != 0);
+				});
+
+				if (MeshData)
 				{
-					TVertexInstanceAttributesConstRef<FVector2f> VertexInstanceUVs = FStaticMeshConstAttributes(*MeshData->MeshDescription).GetVertexInstanceUVs();
-					MergeData.NewUVs.Reset(MeshData->MeshDescription->VertexInstances().Num());
-					for (const FVertexInstanceID VertexInstanceID : MeshData->MeshDescription->VertexInstances().GetElementIDs())
+					if (MeshData->CustomTextureCoordinates.Num())
+					{
+						MergeData.NewUVs = MeshData->CustomTextureCoordinates;
+					}
+					else
 					{
-						MergeData.NewUVs.Add(FVector2D(VertexInstanceUVs.Get(VertexInstanceID, MeshData->TextureCoordinateIndex)));
+						TVertexInstanceAttributesConstRef<FVector2f> VertexInstanceUVs = FStaticMeshConstAttributes(*MeshData->MeshDescription).GetVertexInstanceUVs();
+						MergeData.NewUVs.Reset(MeshData->MeshDescription->VertexInstances().Num());
+						for (const FVertexInstanceID VertexInstanceID : MeshData->MeshDescription->VertexInstances().GetElementIDs())
+						{
+							MergeData.NewUVs.Add(FVector2D(VertexInstanceUVs.Get(VertexInstanceID, MeshData->TextureCoordinateIndex)));
+						}
 					}
+					MergeData.TexCoordBounds[0] = FBox2D(FVector2D(0.0f, 0.0f), FVector2D(1.0f, 1.0f));
 				}
-				MergeData.TexCoordBounds[0] = FBox2D(FVector2D(0.0f, 0.0f), FVector2D(1.0f, 1.0f));
 			}
 			MergeDataEntries.Add(MergeData);
 		}
