Removed public module dependencies because they aren't used by the module making them public and they create circular dependency issues.

diff --git a/Engine/Plugins/Animation/AnimationData/Source/AnimationData/AnimationData.build.cs b/Engine/Plugins/Animation/AnimationData/Source/AnimationData/AnimationData.build.cs
--- a/Engine/Plugins/Animation/AnimationData/Source/AnimationData/AnimationData.build.cs
+++ b/Engine/Plugins/Animation/AnimationData/Source/AnimationData/AnimationData.build.cs
@@ -21,7 +21,8 @@ public class AnimationData : ModuleRules
                 "AnimationDataController",
                 "AnimationCore",
                 "AnimationBlueprintLibrary",
-                "TimeManagement"
+				"AnimationEditor",
+				"TimeManagement"
             });
 
         PublicDependencyModuleNames.AddRange(
diff --git a/Engine/Plugins/Animation/IKRig/Source/IKRig/IKRig.Build.cs b/Engine/Plugins/Animation/IKRig/Source/IKRig/IKRig.Build.cs
--- a/Engine/Plugins/Animation/IKRig/Source/IKRig/IKRig.Build.cs
+++ b/Engine/Plugins/Animation/IKRig/Source/IKRig/IKRig.Build.cs
@@ -39,6 +39,12 @@ public class IKRig : ModuleRules
 					"ApplicationCore",
 					"MessageLog",
 				});
+
+			PublicIncludePathModuleNames.AddRange(
+				new string[]
+				{
+					"AnimationWidgets",
+				});
 		}
 	}
 }
diff --git a/Engine/Plugins/Animation/MLDeformer/MLDeformerFramework/Source/MLDeformerFrameworkEditor/MLDeformerFrameworkEditor.Build.cs b/Engine/Plugins/Animation/MLDeformer/MLDeformerFramework/Source/MLDeformerFrameworkEditor/MLDeformerFrameworkEditor.Build.cs
--- a/Engine/Plugins/Animation/MLDeformer/MLDeformerFramework/Source/MLDeformerFrameworkEditor/MLDeformerFrameworkEditor.Build.cs
+++ b/Engine/Plugins/Animation/MLDeformer/MLDeformerFramework/Source/MLDeformerFrameworkEditor/MLDeformerFrameworkEditor.Build.cs
@@ -57,9 +57,10 @@ namespace UnrealBuildTool.Rules
 				}
 			);
 
-			PrivateIncludePathModuleNames.AddRange(
+			PublicIncludePathModuleNames.AddRange(
 				new string[]
 				{
+					"Persona",
 				}
 			);
 		}
diff --git a/Engine/Plugins/Editor/EngineAssetDefinitions/Source/EngineAssetDefinitions.Build.cs b/Engine/Plugins/Editor/EngineAssetDefinitions/Source/EngineAssetDefinitions.Build.cs
--- a/Engine/Plugins/Editor/EngineAssetDefinitions/Source/EngineAssetDefinitions.Build.cs
+++ b/Engine/Plugins/Editor/EngineAssetDefinitions/Source/EngineAssetDefinitions.Build.cs
@@ -11,8 +11,15 @@ public class EngineAssetDefinitions : ModuleRules
 			{
 				"Core",
 			}
-		);			
-		
+		);
+
+		PublicIncludePathModuleNames.AddRange(
+			new string[]
+			{
+				"SkeletalMeshEditor",
+			}
+		);
+
 		PrivateDependencyModuleNames.AddRange(
 			new string[]
 			{
diff --git a/Engine/Plugins/Experimental/Animation/PoseSearch/Source/Editor/PoseSearchEditor.Build.cs b/Engine/Plugins/Experimental/Animation/PoseSearch/Source/Editor/PoseSearchEditor.Build.cs
--- a/Engine/Plugins/Experimental/Animation/PoseSearch/Source/Editor/PoseSearchEditor.Build.cs
+++ b/Engine/Plugins/Experimental/Animation/PoseSearch/Source/Editor/PoseSearchEditor.Build.cs
@@ -41,6 +41,7 @@ public class PoseSearchEditor : ModuleRules
 		PrivateDependencyModuleNames.AddRange(
 			new string[]
 			{
+				"AnimationEditor",
 				"AssetDefinition",
 				"BlueprintGraph",
 				"EditorStyle",
diff --git a/Engine/Plugins/Experimental/Animation/SkeletalMeshModelingTools/Source/SkeletalMeshModelingTools/SkeletalMeshModelingTools.Build.cs b/Engine/Plugins/Experimental/Animation/SkeletalMeshModelingTools/Source/SkeletalMeshModelingTools/SkeletalMeshModelingTools.Build.cs
--- a/Engine/Plugins/Experimental/Animation/SkeletalMeshModelingTools/Source/SkeletalMeshModelingTools/SkeletalMeshModelingTools.Build.cs
+++ b/Engine/Plugins/Experimental/Animation/SkeletalMeshModelingTools/Source/SkeletalMeshModelingTools/SkeletalMeshModelingTools.Build.cs
@@ -23,6 +23,7 @@ public class SkeletalMeshModelingTools : ModuleRules
 				"MeshModelingToolsExp",
 				"ModelingComponentsEditorOnly",
 				"ModelingToolsEditorMode",
+				"SkeletalMeshEditor",
 				"SkeletalMeshUtilitiesCommon",
 				"Slate",
 				"SlateCore",
diff --git a/Engine/Plugins/Experimental/Dataflow/Source/DataflowEditor/DataflowEditor.Build.cs b/Engine/Plugins/Experimental/Dataflow/Source/DataflowEditor/DataflowEditor.Build.cs
--- a/Engine/Plugins/Experimental/Dataflow/Source/DataflowEditor/DataflowEditor.Build.cs
+++ b/Engine/Plugins/Experimental/Dataflow/Source/DataflowEditor/DataflowEditor.Build.cs
@@ -6,6 +6,13 @@ namespace UnrealBuildTool.Rules
 	{
         public DataflowEditor(ReadOnlyTargetRules Target) : base(Target)
 		{
+			PublicIncludePathModuleNames.AddRange(
+				new string[]
+				{
+					"SkeletonEditor"
+				}
+			);
+
 			PublicDependencyModuleNames.AddRange(
 				new string[]
 				{
diff --git a/Engine/Plugins/Experimental/MeshModelingToolsetExp/Source/ModelingEditorUI/ModelingEditorUI.build.cs b/Engine/Plugins/Experimental/MeshModelingToolsetExp/Source/ModelingEditorUI/ModelingEditorUI.build.cs
--- a/Engine/Plugins/Experimental/MeshModelingToolsetExp/Source/ModelingEditorUI/ModelingEditorUI.build.cs
+++ b/Engine/Plugins/Experimental/MeshModelingToolsetExp/Source/ModelingEditorUI/ModelingEditorUI.build.cs
@@ -38,7 +38,8 @@ public class ModelingEditorUI : ModuleRules
 				"ToolWidgets",
 				"EditorWidgets",
 				"DeveloperSettings", 
-				"MeshModelingToolsExp"
+				"MeshModelingToolsExp",
+				"Persona"
 			}
 			);
 
diff --git a/Engine/Source/Developer/MeshMergeUtilities/MeshMergeUtilities.Build.cs b/Engine/Source/Developer/MeshMergeUtilities/MeshMergeUtilities.Build.cs
--- a/Engine/Source/Developer/MeshMergeUtilities/MeshMergeUtilities.Build.cs
+++ b/Engine/Source/Developer/MeshMergeUtilities/MeshMergeUtilities.Build.cs
@@ -38,13 +38,15 @@ public class MeshMergeUtilities : ModuleRules
 
         PublicIncludePathModuleNames.AddRange(
           new string[] {
-               "HierarchicalLODUtilities",
-               "MeshUtilities",
                "MeshReductionInterface",
-			   "SkeletalMeshEditor",
-			   "Persona",
           }
         );
+		
+		PrivateIncludePathModuleNames.AddRange(
+			new string [] {
+				"HierarchicalLODUtilities"
+			}
+		);
 
         DynamicallyLoadedModuleNames.AddRange(
             new string[] {
diff --git a/Engine/Source/Editor/AnimationEditor/AnimationEditor.Build.cs b/Engine/Source/Editor/AnimationEditor/AnimationEditor.Build.cs
--- a/Engine/Source/Editor/AnimationEditor/AnimationEditor.Build.cs
+++ b/Engine/Source/Editor/AnimationEditor/AnimationEditor.Build.cs
@@ -32,7 +32,13 @@ public class AnimationEditor : ModuleRules
             }
         );
 
-        DynamicallyLoadedModuleNames.AddRange(
+		PublicIncludePathModuleNames.AddRange(
+			new string[] {
+				"Persona",
+			}
+		);
+
+		DynamicallyLoadedModuleNames.AddRange(
             new string[] {
                 "PropertyEditor",
                 "SequenceRecorder",
diff --git a/Engine/Source/Editor/GraphEditor/GraphEditor.Build.cs b/Engine/Source/Editor/GraphEditor/GraphEditor.Build.cs
--- a/Engine/Source/Editor/GraphEditor/GraphEditor.Build.cs
+++ b/Engine/Source/Editor/GraphEditor/GraphEditor.Build.cs
@@ -32,6 +32,7 @@ public class GraphEditor : ModuleRules
 				"KismetWidgets",
 				"BlueprintGraph",
 				"Documentation",
+				"Persona",
 				"RenderCore",
 				"RHI",
 				"ToolMenus",
diff --git a/Engine/Source/Runtime/MovieSceneCapture/MovieSceneCapture.Build.cs b/Engine/Source/Runtime/MovieSceneCapture/MovieSceneCapture.Build.cs
--- a/Engine/Source/Runtime/MovieSceneCapture/MovieSceneCapture.Build.cs
+++ b/Engine/Source/Runtime/MovieSceneCapture/MovieSceneCapture.Build.cs
@@ -12,13 +12,6 @@ public class MovieSceneCapture : ModuleRules
 			DynamicallyLoadedModuleNames.Add("ImageWrapper");
 		}
 
-		PublicDependencyModuleNames.AddRange(
-			new string[] {
-				"LevelSequence",
-				"TimeManagement",
-			}
-		);
-
 		PrivateDependencyModuleNames.AddRange(
 			new string[] {
                 "AssetRegistry",
