Merge Actor - Mesh Merge: Fix corner case where a merged mesh section was assigned an invalid material index (INDEX_NONE)
* Required a source mesh with two sections with identical materials
* Occured only when created a merged mesh without baking down textures

diff --git a/Engine/Source/Developer/MeshMergeUtilities/Private/MeshMergeUtilities.cpp b/Engine/Source/Developer/MeshMergeUtilities/Private/MeshMergeUtilities.cpp
--- a/Engine/Source/Developer/MeshMergeUtilities/Private/MeshMergeUtilities.cpp
+++ b/Engine/Source/Developer/MeshMergeUtilities/Private/MeshMergeUtilities.cpp
@@ -2447,7 +2447,7 @@ void FMeshMergeUtilities::MergeComponentsToStaticMesh(const TArray<UPrimitiveCom
 				FMeshSectionInfo SectionInfo;
 				SectionInfo.bCastShadow = StoredSectionInfo.EnabledProperties.Contains(GET_MEMBER_NAME_CHECKED(FMeshSectionInfo, bCastShadow));
 				SectionInfo.bEnableCollision = StoredSectionInfo.EnabledProperties.Contains(GET_MEMBER_NAME_CHECKED(FMeshSectionInfo, bEnableCollision));
-				SectionInfo.MaterialIndex = UniqueMaterials.IndexOfByKey(StoredSectionInfo.Material);
+				SectionInfo.MaterialIndex = UniqueMaterials.Num() == 1 ? 0 : UniqueMaterials.IndexOfByKey(CollapsedMaterialMap[StoredSectionInfo.Material]);
 				SectionInfoMap.Set(LODIndex, Index, SectionInfo);
 			}
 		}
