Merge Merging - Fix crashes occurring when a GC was performed during mesh merging

diff --git a/Engine/Source/Developer/MeshMergeUtilities/Private/MeshMergeUtilities.cpp b/Engine/Source/Developer/MeshMergeUtilities/Private/MeshMergeUtilities.cpp
--- a/Engine/Source/Developer/MeshMergeUtilities/Private/MeshMergeUtilities.cpp
+++ b/Engine/Source/Developer/MeshMergeUtilities/Private/MeshMergeUtilities.cpp
@@ -83,6 +83,8 @@
 #include "ISMPartition/ISMComponentBatcher.h"
 #include "ISMPartition/ISMComponentDescriptor.h"
 
+#include "UObject/GCObjectScopeGuard.h"
+
 #define LOCTEXT_NAMESPACE "MeshMergeUtils"
 
 DEFINE_LOG_CATEGORY(LogMeshMerging);
@@ -1675,6 +1677,8 @@ void FMeshMergeUtilities::CreateProxyMesh(const TArray<UStaticMeshComponent*>& I
 	}
 
 	UMaterialOptions* Options = PopulateMaterialOptions(MaterialProxySettings);
+	TGCObjectScopeGuard<UMaterialOptions> MaterialOptionsGCScopeGuard(Options);
+
 	TArray<EMaterialProperty> MaterialProperties;
 	for (const FPropertyEntry& Entry : Options->Properties)
 	{
@@ -2713,6 +2717,7 @@ void FMeshMergeUtilities::CreateMergedMaterial(FMeshMergeDataTracker& InDataTrac
 	}
 
 	UMaterialOptions* MaterialOptions = PopulateMaterialOptions(MaterialProxySettings);
+	TGCObjectScopeGuard<UMaterialOptions> MaterialOptionsGCScopeGuard(MaterialOptions);
 
 	// Check each material to see if the shader actually uses vertex data and collect flags
 	TArray<TOptional<bool>> bMaterialUsesVertexData;
