Merge Actor - Properly handle custom primitive data when baking materials
* Forward custom primitive data info to material baking
* Added Custom Primitive Data case to material baking engine test
* Also updated broken Material Baking tests screenshots as subtle differences are now occuring on shadows... Disabled shadows from the map light source to minimize noise in future test results

#virtualized

diff --git a/Engine/Plugins/Editor/WorldPartitionHLODUtilities/Source/Private/WorldPartition/HLOD/Builders/HLODBuilderMeshMerge.cpp b/Engine/Plugins/Editor/WorldPartitionHLODUtilities/Source/Private/WorldPartition/HLOD/Builders/HLODBuilderMeshMerge.cpp
--- a/Engine/Plugins/Editor/WorldPartitionHLODUtilities/Source/Private/WorldPartition/HLOD/Builders/HLODBuilderMeshMerge.cpp
+++ b/Engine/Plugins/Editor/WorldPartitionHLODUtilities/Source/Private/WorldPartition/HLOD/Builders/HLODBuilderMeshMerge.cpp
@@ -43,7 +43,7 @@ uint32 UHLODBuilderMeshMergeSettings::GetCRC() const
 	FArchiveCrc32 Ar;
 
 	// Base mesh merge key, changing this will force a rebuild of all HLODs from this builder
-	FString HLODBaseKey = "89D89284DD3847FA90C5998E06DD8FEC";
+	FString HLODBaseKey = "B8DB9CB1780C4EE1B80A36D9E205AA0F";
 	Ar << HLODBaseKey;
 
 	Ar << This.MeshMergeSettings;
diff --git a/Engine/Plugins/Enterprise/GLTFExporter/Source/GLTFExporter/Private/Converters/GLTFMaterialUtilities.cpp b/Engine/Plugins/Enterprise/GLTFExporter/Source/GLTFExporter/Private/Converters/GLTFMaterialUtilities.cpp
--- a/Engine/Plugins/Enterprise/GLTFExporter/Source/GLTFExporter/Private/Converters/GLTFMaterialUtilities.cpp
+++ b/Engine/Plugins/Enterprise/GLTFExporter/Source/GLTFExporter/Private/Converters/GLTFMaterialUtilities.cpp
@@ -138,7 +138,7 @@ FGLTFPropertyBakeOutput FGLTFMaterialUtilities::BakeMaterialProperty(const FIntP
 		MeshSet.LightMap = MeshData->LightMap;
 		MeshSet.LightMapIndex = MeshData->LightMapTexCoord;
 		MeshSet.LightmapResourceCluster = MeshData->LightMapResourceCluster;
-		MeshSet.PrimitiveData = &MeshData->PrimitiveData;
+		MeshSet.PrimitiveData = MeshData->PrimitiveData;
 	}
 
 	FMaterialDataEx MatSet;
diff --git a/Engine/Source/Developer/MaterialBaking/Private/MaterialRenderItem.cpp b/Engine/Source/Developer/MaterialBaking/Private/MaterialRenderItem.cpp
--- a/Engine/Source/Developer/MaterialBaking/Private/MaterialRenderItem.cpp
+++ b/Engine/Source/Developer/MaterialBaking/Private/MaterialRenderItem.cpp
@@ -121,7 +121,8 @@ void FMeshMaterialRenderItem::QueueMaterial(FCanvasRenderContext& RenderContext,
 			DynamicMeshBuilder.AddTriangles(Indices);
 		}
 
-		FPrimitiveData PrimitiveData = MeshSettings->PrimitiveData != nullptr ? *MeshSettings->PrimitiveData : FPrimitiveData();
+		const FPrimitiveData DefaultPrimitiveData;
+		const FPrimitiveData& PrimitiveData = MeshSettings->PrimitiveData.Get(DefaultPrimitiveData);
 		FPrimitiveUniformShaderParameters PrimitiveParams = FPrimitiveUniformShaderParametersBuilder{}
 			.Defaults()
 				.LocalToWorld(PrimitiveData.LocalToWorld)
@@ -163,7 +164,7 @@ void FMeshMaterialRenderItem::QueueMaterial(FCanvasRenderContext& RenderContext,
 void FMeshMaterialRenderItem::PopulateWithQuadData()
 {
 	// Pre-transform all vertices with the inverse of LocalToWorld to negate its effect during material baking 
-	const FMatrix44f WorldToLocal = (MeshSettings->PrimitiveData != nullptr) ? FMatrix44f(MeshSettings->PrimitiveData->LocalToWorld.Inverse()) : FMatrix44f::Identity;
+	const FMatrix44f WorldToLocal = MeshSettings->PrimitiveData.IsSet() ? FMatrix44f(MeshSettings->PrimitiveData->LocalToWorld.Inverse()) : FMatrix44f::Identity;
 	
 	Vertices.Empty(4);
 	Indices.Empty(6);
@@ -202,7 +203,7 @@ void FMeshMaterialRenderItem::PopulateWithQuadData()
 void FMeshMaterialRenderItem::PopulateWithMeshData()
 {
 	// Pre-transform all vertices with the inverse of LocalToWorld to negate its effect during material baking 
-	const FMatrix44f WorldToLocal = (MeshSettings->PrimitiveData != nullptr) ? FMatrix44f(MeshSettings->PrimitiveData->LocalToWorld.Inverse()) : FMatrix44f::Identity;
+	const FMatrix44f WorldToLocal = MeshSettings->PrimitiveData.IsSet() ? FMatrix44f(MeshSettings->PrimitiveData->LocalToWorld.Inverse()) : FMatrix44f::Identity;
 
 	const FMeshDescription* RawMesh = MeshSettings->MeshDescription;
 
diff --git a/Engine/Source/Developer/MaterialBaking/Public/MaterialBakingStructures.h b/Engine/Source/Developer/MaterialBaking/Public/MaterialBakingStructures.h
--- a/Engine/Source/Developer/MaterialBaking/Public/MaterialBakingStructures.h
+++ b/Engine/Source/Developer/MaterialBaking/Public/MaterialBakingStructures.h
@@ -150,8 +150,8 @@ struct FMeshData
 	/** Pointer to the LightmapResourceCluster to be passed on the the LightCacheInterface when baking */
 	const FLightmapResourceCluster* LightmapResourceCluster = nullptr;
 
-	/** Pointer to primitive data that is accessible through material expressions, if nullptr default values are used */
-	const FPrimitiveData* PrimitiveData = nullptr;
+	/** Optional primitive data that is accessible through material expressions */
+	TOptional<FPrimitiveData> PrimitiveData;
 };
 
 /** Structure containing data being processed while baking out materials*/
diff --git a/Engine/Source/Developer/MeshMergeUtilities/Private/MeshMergeUtilities.cpp b/Engine/Source/Developer/MeshMergeUtilities/Private/MeshMergeUtilities.cpp
--- a/Engine/Source/Developer/MeshMergeUtilities/Private/MeshMergeUtilities.cpp
+++ b/Engine/Source/Developer/MeshMergeUtilities/Private/MeshMergeUtilities.cpp
@@ -2804,6 +2804,11 @@ void FMeshMergeUtilities::CreateMergedMaterial(FMeshMergeDataTracker& InDataTrac
 				NewMeshData.VertexColorHash = Key.GetVertexColorHash();
 				NewMeshData.bMirrored = Component->GetComponentTransform().GetDeterminant() < 0.0f;
 				NewMeshData.MaterialIndices = SectionIndices;
+				if (!Component->GetCustomPrimitiveData().Data.IsEmpty())
+				{
+					NewMeshData.PrimitiveData = FPrimitiveData();
+					NewMeshData.PrimitiveData->CustomPrimitiveData = &Component->GetCustomPrimitiveData();
+				}
 			}
 
 			auto CompareMaterialData = [&InSettings](const FMaterialData& LHS, const FMaterialData& RHS)
@@ -2811,9 +2816,21 @@ void FMeshMergeUtilities::CreateMergedMaterial(FMeshMergeDataTracker& InDataTrac
 				return InSettings.bMergeEquivalentMaterials ? FMaterialKey(LHS.Material) == FMaterialKey(RHS.Material) : LHS.Material == RHS.Material;
 			};
 
-			auto CompareMeshData = [](const FMeshData& LHS, const FMeshData& RHS)
+			auto CompareCustomPrimitiveData = [](const FCustomPrimitiveData* LHS, const FCustomPrimitiveData* RHS)
+			{
+				// Return true if both are null, false if one of them is null - otherwise, compare content
+				return (!LHS && !RHS) ? true : (!LHS || !RHS) ? false : (*LHS == *RHS);
+			};
+
+			auto ComparePrimitiveData = [&CompareCustomPrimitiveData](const TOptional<FPrimitiveData>& LHS, const TOptional<FPrimitiveData>& RHS)
+			{
+				// Return true if both are null, false if one of them is null - otherwise, compare content
+				return (!LHS && !RHS) ? true : (!LHS || !RHS) ? false : CompareCustomPrimitiveData(LHS->CustomPrimitiveData, RHS->CustomPrimitiveData);
+			};
+
+			auto CompareMeshData = [&ComparePrimitiveData](const FMeshData& LHS, const FMeshData& RHS)
 			{
-				return (LHS.Mesh == RHS.Mesh) && (LHS.MaterialIndices == RHS.MaterialIndices) && (LHS.bMirrored == RHS.bMirrored) && (LHS.VertexColorHash == RHS.VertexColorHash);
+				return (LHS.Mesh == RHS.Mesh) && (LHS.MaterialIndices == RHS.MaterialIndices) && (LHS.bMirrored == RHS.bMirrored) && (LHS.VertexColorHash == RHS.VertexColorHash) && ComparePrimitiveData(LHS.PrimitiveData, RHS.PrimitiveData);
 			};
 
 			// Find material & mesh pair
diff --git a/Engine/Source/Developer/MeshMergeUtilities/Private/StaticMeshComponentAdapter.cpp b/Engine/Source/Developer/MeshMergeUtilities/Private/StaticMeshComponentAdapter.cpp
--- a/Engine/Source/Developer/MeshMergeUtilities/Private/StaticMeshComponentAdapter.cpp
+++ b/Engine/Source/Developer/MeshMergeUtilities/Private/StaticMeshComponentAdapter.cpp
@@ -52,6 +52,16 @@ void FStaticMeshComponentAdapter::ApplySettings(int32 LODIndex, FMeshData& InOut
 			InOutMeshData.LightmapResourceCluster = MeshMapBuildData->ResourceCluster;
 		}
 	}
+
+	if (!StaticMeshComponent->GetCustomPrimitiveData().Data.IsEmpty())
+	{
+		if (!InOutMeshData.PrimitiveData.IsSet())
+		{
+			InOutMeshData.PrimitiveData = FPrimitiveData();
+		}
+
+		InOutMeshData.PrimitiveData->CustomPrimitiveData = &StaticMeshComponent->GetCustomPrimitiveData();
+	}
 }
 
 UPackage* FStaticMeshComponentAdapter::GetOuter() const
