MergeMergeUtility - Fix broken hash computation in FProxyMeshDescriptor
* Return values of TypeCrc32() was ignored in some cases

diff --git a/Engine/Source/Developer/MeshMergeUtilities/Private/MeshMergeUtilities.cpp b/Engine/Source/Developer/MeshMergeUtilities/Private/MeshMergeUtilities.cpp
--- a/Engine/Source/Developer/MeshMergeUtilities/Private/MeshMergeUtilities.cpp
+++ b/Engine/Source/Developer/MeshMergeUtilities/Private/MeshMergeUtilities.cpp
@@ -1157,7 +1157,6 @@ public:
 		, LightMapIndex(INDEX_NONE)
 	{
 		ISMDescriptor.InitFrom(StaticMeshComponent, false);
-		ISMDescriptor.ComputeHash();
 
 		// Retrieve lightmap for usage of lightmap data
 		if (StaticMeshComponent->LODData.IsValidIndex(0))
@@ -1170,16 +1169,14 @@ public:
 				LightMapIndex = StaticMeshComponent->GetStaticMesh()->GetLightMapCoordinateIndex();
 			}
 		}
-
 		
 		Hash = ISMDescriptor.ComputeHash();
-
-		FCrc::TypeCrc32(LODIndex, Hash);
+		Hash = FCrc::TypeCrc32(LODIndex, Hash);
 
 		if (LightMapIndex != INDEX_NONE)
 		{
-			FCrc::TypeCrc32(LightMap.GetReference(), Hash);
-			FCrc::TypeCrc32(LightMapIndex, Hash);
+			Hash = FCrc::TypeCrc32(LightMap.GetReference(), Hash);
+			Hash = FCrc::TypeCrc32(LightMapIndex, Hash);
 		}
 	}
 
