Merge Actor - Mesh Merging: Correctly handle MID when merging meshes but not their materials
* Duplicate the MID and outer it to the generated static mesh

diff --git a/Engine/Source/Developer/MeshMergeUtilities/Private/MeshMergeUtilities.cpp b/Engine/Source/Developer/MeshMergeUtilities/Private/MeshMergeUtilities.cpp
--- a/Engine/Source/Developer/MeshMergeUtilities/Private/MeshMergeUtilities.cpp
+++ b/Engine/Source/Developer/MeshMergeUtilities/Private/MeshMergeUtilities.cpp
@@ -23,6 +23,7 @@
 #include "UObject/UObjectBaseUtility.h"
 #include "UObject/Package.h"
 #include "Materials/Material.h"
+#include "Materials/MaterialInstanceDynamic.h"
 #include "Misc/ScopedSlowTask.h"
 #include "Modules/ModuleManager.h"
 #include "HierarchicalLODUtilitiesModule.h"
@@ -2619,8 +2620,17 @@ void FMeshMergeUtilities::MergeComponentsToStaticMesh(const TArray<UPrimitiveCom
 		{
 			if (Material && (!Material->IsAsset() && InOuter != GetTransientPackage()))
 			{
-				Material = nullptr; // do not save non-asset materials
+				// MIDs are not assets, duplicate them and outer them to the static mesh.
+				if (UMaterialInstanceDynamic* MID = Cast<UMaterialInstanceDynamic>(Material))
+				{
+					Material = DuplicateObject<UMaterialInstanceDynamic>(MID, StaticMesh);
+				}
+				else
+				{
+					Material = nullptr; // do not save non-asset materials
+				}
 			}
+
 			//Make sure we have unique slot name here
 			FName MaterialSlotName = DataTracker.GetMaterialSlotName(Material);
 			int32 Counter = 1;
