Mesh Merge - Prevent invalid mesh inputs from generating a static mesh asset with no source data

diff --git a/Engine/Source/Developer/MeshMergeUtilities/Private/MeshMergeUtilities.cpp b/Engine/Source/Developer/MeshMergeUtilities/Private/MeshMergeUtilities.cpp
--- a/Engine/Source/Developer/MeshMergeUtilities/Private/MeshMergeUtilities.cpp
+++ b/Engine/Source/Developer/MeshMergeUtilities/Private/MeshMergeUtilities.cpp
@@ -88,6 +88,8 @@
 
 #include "UObject/GCObjectScopeGuard.h"
 
+#include "Algo/RemoveIf.h"
+
 #define LOCTEXT_NAMESPACE "MeshMergeUtils"
 
 DEFINE_LOG_CATEGORY(LogMeshMerging);
@@ -2486,6 +2488,13 @@ void FMeshMergeUtilities::MergeComponentsToStaticMesh(const TArray<UPrimitiveCom
 	//
 	//Create merged mesh asset
 	//
+
+	MergedRawMeshes.SetNum(Algo::RemoveIf(MergedRawMeshes, [](const FMeshDescription& MeshDescription) { return MeshDescription.IsEmpty(); }));
+
+	const bool bContainsImposters = !ImposterComponents.IsEmpty();
+	const bool bContainsMergedMeshes = !MergedRawMeshes.IsEmpty();
+
+	if (bContainsMergedMeshes || bContainsImposters)
 	{
 		FString AssetName;
 		FString PackageName;
@@ -2553,7 +2562,6 @@ void FMeshMergeUtilities::MergeComponentsToStaticMesh(const TArray<UPrimitiveCom
 		// Ray tracing support
 		StaticMesh->bSupportRayTracing = InSettings.bSupportRayTracing;
 
-		const bool bContainsImposters = ImposterComponents.Num() > 0;
 		TArray<UMaterialInterface*> ImposterMaterials;
 		FBox ImposterBounds(EForceInit::ForceInit);
 		for (int32 LODIndex = 0; LODIndex < MergedRawMeshes.Num(); ++LODIndex)
@@ -2725,6 +2733,10 @@ void FMeshMergeUtilities::MergeComponentsToStaticMesh(const TArray<UPrimitiveCom
 		OutAssetsToSync.Add(StaticMesh);
 		OutMergedActorLocation = MergedAssetPivot;
 	}
+	else
+	{
+		UE_LOG(LogMeshMerging, Display, TEXT("MergeComponentsToStaticMesh: Skipped creation of a static mesh asset as no input meshes were provided"));
+	}
 }
 
 void FMeshMergeUtilities::CreateMergedMaterial(FMeshMergeDataTracker& InDataTracker, const FMeshMergingSettings& InSettings, const TArray<UStaticMeshComponent*>& InStaticMeshComponentsToMerge, TArray<FStaticMeshComponentAdapter>& InAdapters, const TArray<UMaterialInterface*>& InUniqueMaterials, const TMap<UMaterialInterface*, UMaterialInterface*>& InCollapsedMaterialMap, TMultiMap<FMeshLODKey, MaterialRemapPair>& InOutputMaterialsMap, bool bInMergeAllLODs, bool bInMergeMaterialData, const FVector& InMergedAssetPivot, FFlattenMaterial& OutFlattenMaterial) const
