ProxyLOD - Properly handle meshes with empty sections

diff --git a/Engine/Source/Developer/MeshMergeUtilities/Private/MeshMergeHelpers.cpp b/Engine/Source/Developer/MeshMergeUtilities/Private/MeshMergeHelpers.cpp
--- a/Engine/Source/Developer/MeshMergeUtilities/Private/MeshMergeHelpers.cpp
+++ b/Engine/Source/Developer/MeshMergeUtilities/Private/MeshMergeHelpers.cpp
@@ -344,11 +344,21 @@ void FMeshMergeHelpers::ExportStaticMeshLOD(const FStaticMeshLODResources& Stati
 	const int32 NumTexCoords = StaticMeshLOD.VertexBuffers.StaticMeshVertexBuffer.GetNumTexCoords();
 	VertexInstanceUVs.SetNumChannels(NumTexCoords);
 
+	TMap<int32, FPolygonGroupID> SectionToPolygonGroupMap;
+
 	for (int32 SectionIndex = 0; SectionIndex < StaticMeshLOD.Sections.Num(); ++SectionIndex)
 	{
 		const FStaticMeshSection& Section = StaticMeshLOD.Sections[SectionIndex];
+
+		// Skip empty sections
+		if (Section.NumTriangles == 0)
+		{
+			continue;
+		}
+
 		FPolygonGroupID CurrentPolygonGroupID = OutMeshDescription.CreatePolygonGroup();
-		check(CurrentPolygonGroupID.GetValue() == SectionIndex);
+		SectionToPolygonGroupMap.Emplace(SectionIndex, CurrentPolygonGroupID);
+
 		if (Materials.IsValidIndex(Section.MaterialIndex))
 		{
 			PolygonGroupImportedMaterialSlotNames[CurrentPolygonGroupID] = Materials[Section.MaterialIndex].ImportedMaterialSlotName;
@@ -377,7 +387,7 @@ void FMeshMergeHelpers::ExportStaticMeshLOD(const FStaticMeshLODResources& Stati
 			uint32 EndTriangle = BeginTriangle + Section.NumTriangles;
 			if ((uint32)TriangleIndex >= BeginTriangle && (uint32)TriangleIndex < EndTriangle)
 			{
-				CurrentPolygonGroupID = FPolygonGroupID(SectionIndex);
+				CurrentPolygonGroupID = SectionToPolygonGroupMap[SectionIndex];
 				break;
 			}
 		}
