Fixed crash occurring when building HLODs after having built light maps

diff --git a/Engine/Source/Developer/MeshMergeUtilities/Private/MeshMergeUtilities.cpp b/Engine/Source/Developer/MeshMergeUtilities/Private/MeshMergeUtilities.cpp
--- a/Engine/Source/Developer/MeshMergeUtilities/Private/MeshMergeUtilities.cpp
+++ b/Engine/Source/Developer/MeshMergeUtilities/Private/MeshMergeUtilities.cpp
@@ -1428,7 +1428,18 @@ TArray<FMeshData> PrepareBakingMeshes(const struct FMeshProxySettings& InMeshPro
 	TArray<FMeshData> MeshData;
 	MeshData.SetNum(InDescriptors.Num());
 
-	// Parallel step
+	// GetLightMap() must be called from the game thread
+	check(IsInGameThread());
+	for (int32 MeshIndex = 0; MeshIndex < InDescriptors.Num(); ++MeshIndex)
+	{
+		if (InDescriptors[MeshIndex].GetLightMapIndex() != INDEX_NONE)
+		{
+			MeshData[MeshIndex].LightMap = InDescriptors[MeshIndex].GetLightMap();
+			MeshData[MeshIndex].LightMapIndex = InDescriptors[MeshIndex].GetLightMapIndex();
+		}
+	}
+
+	// Parallel step - fetching custom (unwrapped) texture coordinates is the slowest part here
 	ParallelFor(InDescriptors.Num(), [&MeshData, &InDescriptors, &InMeshDescriptionData, &InMeshProxySettings](uint32 MeshIndex)
 	{
 		const FProxyMeshDescriptor& MeshDescriptor = InDescriptors[MeshIndex];
@@ -1436,13 +1447,6 @@ TArray<FMeshData> PrepareBakingMeshes(const struct FMeshProxySettings& InMeshPro
 		FMeshData& MeshSettings = MeshData[MeshIndex];
 		MeshSettings.TextureCoordinateBox = FBox2D(FVector2D(0.0f, 0.0f), FVector2D(1.0f, 1.0f));
 
-
-		if (MeshDescriptor.GetLightMapIndex() != INDEX_NONE)
-		{
-			MeshSettings.LightMap = MeshDescriptor.GetLightMap();
-			MeshSettings.LightMapIndex = MeshDescriptor.GetLightMapIndex();
-		}
-
 		if (InMeshProxySettings.bGroupIdenticalMeshesForBaking)
 		{
 			// Grouping by identical meshes, the UVs should have already been setup
