Mesh Merge: Fix potentially invalid material baking result when using settings to strip out UV channels from the final mesh
* All valid UV channels must be kept during material baking
#tests enginetest

diff --git a/Engine/Plugins/Editor/WorldPartitionHLODUtilities/Source/Private/WorldPartition/HLOD/Builders/HLODBuilderMeshMerge.cpp b/Engine/Plugins/Editor/WorldPartitionHLODUtilities/Source/Private/WorldPartition/HLOD/Builders/HLODBuilderMeshMerge.cpp
--- a/Engine/Plugins/Editor/WorldPartitionHLODUtilities/Source/Private/WorldPartition/HLOD/Builders/HLODBuilderMeshMerge.cpp
+++ b/Engine/Plugins/Editor/WorldPartitionHLODUtilities/Source/Private/WorldPartition/HLOD/Builders/HLODBuilderMeshMerge.cpp
@@ -59,7 +59,7 @@ uint32 UHLODBuilderMeshMergeSettings::GetCRC() const
 	FArchiveCrc32 Ar;
 
 	// Base mesh merge key, changing this will force a rebuild of all HLODs from this builder
-	FString HLODBaseKey = "2126404D867E40B7BC8B59C408FAF8AA";
+	FString HLODBaseKey = "F906DCC49BCB4102A821B20D64C6E93B";
 	Ar << HLODBaseKey;
 
 	Ar << This.MeshMergeSettings;
diff --git a/Engine/Source/Developer/MeshMergeUtilities/Private/MeshMergeUtilities.cpp b/Engine/Source/Developer/MeshMergeUtilities/Private/MeshMergeUtilities.cpp
--- a/Engine/Source/Developer/MeshMergeUtilities/Private/MeshMergeUtilities.cpp
+++ b/Engine/Source/Developer/MeshMergeUtilities/Private/MeshMergeUtilities.cpp
@@ -2971,10 +2971,19 @@ void FMeshMergeUtilities::CreateMergedMaterial(FMeshMergeDataTracker& InDataTrac
 
 	if (bGloballyRemapUVs)
 	{
-		// We must keep vertex data in order to properly generate unique UVs
+		// Adjust merge settings when merging for the unique UV/material baking pass
+		// The final merged mesh will use the original settings
 		FMeshMergingSettings RemapUVMergeSettings = InSettings;
+		
+		// Keep vertex data in order to properly generate unique UVs
 		RemapUVMergeSettings.bBakeVertexDataToMesh = true;
 
+		// Keep all UVs as some channels might be needed to properly bake the material
+		for (EUVOutput& OutputUV : RemapUVMergeSettings.OutputUVs)
+		{
+			OutputUV = EUVOutput::OutputChannel;
+		}
+		
 		TArray<FMeshDescription> MergedRawMeshes;
 		CreateMergedRawMeshes(InDataTracker, RemapUVMergeSettings, InStaticMeshComponentsToMerge, InUniqueMaterials, InCollapsedMaterialMap, InOutputMaterialsMap, false, false, InMergedAssetPivot, MergedRawMeshes);
 
diff --git a/Engine/Source/Runtime/Engine/Private/HLODProxy.cpp b/Engine/Source/Runtime/Engine/Private/HLODProxy.cpp
--- a/Engine/Source/Runtime/Engine/Private/HLODProxy.cpp
+++ b/Engine/Source/Runtime/Engine/Private/HLODProxy.cpp
@@ -496,7 +496,7 @@ uint32 UHLODProxy::GetCRC(UStaticMeshComponent* InComponent, uint32 InCRC, const
 }
 
 // Key that forms the basis of the HLOD proxy key. Bump this key (i.e. generate a new GUID) when you want to force a rebuild of ALL HLOD proxies
-#define HLOD_PROXY_BASE_KEY		TEXT("136F4B1AD66E47808C62C1F9CC87CC1F")
+#define HLOD_PROXY_BASE_KEY		TEXT("20D42B32B31542FDA1DAF99D4DA83601")
 
 FName UHLODProxy::GenerateKeyForActor(const ALODActor* LODActor, bool bMustUndoLevelTransform)
 {
