substrate - fix for compilation warning.

#rnx

diff --git a/Engine/Source/Developer/MeshMergeUtilities/Private/MeshMergeUtilities.cpp b/Engine/Source/Developer/MeshMergeUtilities/Private/MeshMergeUtilities.cpp
--- a/Engine/Source/Developer/MeshMergeUtilities/Private/MeshMergeUtilities.cpp
+++ b/Engine/Source/Developer/MeshMergeUtilities/Private/MeshMergeUtilities.cpp
@@ -319,7 +319,7 @@ void FMeshMergeUtilities::BakeMaterialsForComponent(TArray<TWeakObjectPtr<UObjec
 					// SUBSTRATE_TODO: gather information about the material used input, maybe using Material->GetSubstrateCompilationConfig().
 					const float bIsSubstrateEnabled = Substrate::IsSubstrateEnabled();
 
-					if ((bIsSubstrateEnabled || !Entry.bUseConstantValue && Material->IsPropertyActive(Entry.Property)) && Entry.Property != MP_MAX)
+					if ((bIsSubstrateEnabled || (!Entry.bUseConstantValue && Material->IsPropertyActive(Entry.Property))) && Entry.Property != MP_MAX)
 					{
 						MaterialSettings.PropertySizes.Add(Entry.Property, Entry.bUseCustomSize ? Entry.CustomSize : MaterialOptions->TextureSize);
 					}
