Mesh Merge - Prevent crash when trying to merge a static mesh component with an invalid mesh.

diff --git a/Engine/Source/Developer/MeshMergeUtilities/Private/MeshMergeUtilities.cpp b/Engine/Source/Developer/MeshMergeUtilities/Private/MeshMergeUtilities.cpp
--- a/Engine/Source/Developer/MeshMergeUtilities/Private/MeshMergeUtilities.cpp
+++ b/Engine/Source/Developer/MeshMergeUtilities/Private/MeshMergeUtilities.cpp
@@ -2138,7 +2138,7 @@ void FMeshMergeUtilities::MergeComponentsToStaticMesh(const TArray<UPrimitiveCom
 	for (int32 MeshId = 0; MeshId < ComponentsToMerge.Num(); ++MeshId)
 	{
 		UStaticMeshComponent* MeshComponent = Cast<UStaticMeshComponent>(ComponentsToMerge[MeshId]);
-		if (MeshComponent)
+		if (MeshComponent && MeshComponent->GetStaticMesh())
 		{
 			// Make sure referenced lightmaps and shadowmaps are compiled
 			if (MeshComponent->LODData.IsValidIndex(0))
