MeshMergeUtilities - Fix module initialization order issue causing HLOD1 to be missing
* The FSkeletalMeshReduction module can load the MeshUtilities directly upon startup
* In some cases, this caused the FProxyGenerationProcessor module to be initialized before the IMeshReductionManagerModule
* IMeshReductionManagerModule had no assigned MeshMerging interface in that case, causing the MeshMerging->CompleteDelegate registration to be skipped

diff --git a/Engine/Source/Developer/MeshMergeUtilities/Private/ProxyGenerationProcessor.cpp b/Engine/Source/Developer/MeshMergeUtilities/Private/ProxyGenerationProcessor.cpp
--- a/Engine/Source/Developer/MeshMergeUtilities/Private/ProxyGenerationProcessor.cpp
+++ b/Engine/Source/Developer/MeshMergeUtilities/Private/ProxyGenerationProcessor.cpp
@@ -28,6 +28,22 @@ FProxyGenerationProcessor::FProxyGenerationProcessor(const FMeshMergeUtilities*
 	FEditorDelegates::MapChange.AddRaw(this, &FProxyGenerationProcessor::OnMapChange);
 	FEditorDelegates::NewCurrentLevel.AddRaw(this, &FProxyGenerationProcessor::OnNewCurrentLevel);
 	
+	IMeshReductionManagerModule& Module = FModuleManager::Get().LoadModuleChecked<IMeshReductionManagerModule>("MeshReductionInterface");
+	Module.OnModuleConfigurationCompleted(FSimpleDelegate::CreateRaw(this, &FProxyGenerationProcessor::OnModuleConfigurationCompleted));
+#endif // WITH_EDITOR
+}
+
+FProxyGenerationProcessor::~FProxyGenerationProcessor()
+{
+#if WITH_EDITOR
+	FEditorDelegates::MapChange.RemoveAll(this);
+	FEditorDelegates::NewCurrentLevel.RemoveAll(this);
+#endif // WITH_EDITOR
+}
+
+void FProxyGenerationProcessor::OnModuleConfigurationCompleted()
+{
+#if WITH_EDITOR
 	IMeshReductionManagerModule& Module = FModuleManager::Get().LoadModuleChecked<IMeshReductionManagerModule>("MeshReductionInterface");
 	IMeshMerging* MeshMerging = Module.GetMeshMergingInterface();
 	if (!MeshMerging)
@@ -50,15 +66,7 @@ FProxyGenerationProcessor::FProxyGenerationProcessor(const FMeshMergeUtilities*
 		DistributedMeshMerging->CompleteDelegate.BindRaw(this, &FProxyGenerationProcessor::ProxyGenerationComplete);
 		DistributedMeshMerging->FailedDelegate.BindRaw(this, &FProxyGenerationProcessor::ProxyGenerationFailed);
 	}
-#endif // WITH_EDITOR
-}
-
-FProxyGenerationProcessor::~FProxyGenerationProcessor()
-{
-#if WITH_EDITOR
-	FEditorDelegates::MapChange.RemoveAll(this);
-	FEditorDelegates::NewCurrentLevel.RemoveAll(this);
-#endif // WITH_EDITOR
+#endif
 }
 
 void FProxyGenerationProcessor::AddProxyJob(FGuid InJobGuid, FMergeCompleteData* InCompleteData)
diff --git a/Engine/Source/Developer/MeshMergeUtilities/Private/ProxyGenerationProcessor.h b/Engine/Source/Developer/MeshMergeUtilities/Private/ProxyGenerationProcessor.h
--- a/Engine/Source/Developer/MeshMergeUtilities/Private/ProxyGenerationProcessor.h
+++ b/Engine/Source/Developer/MeshMergeUtilities/Private/ProxyGenerationProcessor.h
@@ -46,6 +46,8 @@ protected:
 	};
 
 	void ProcessJob(const FGuid& JobGuid, FProxyGenerationData* Data);
+	void OnModuleConfigurationCompleted();
+
 protected:
 	/** Holds Proxy mesh job data together with the job Guid */
 	TMap<FGuid, FMergeCompleteData*> ProxyMeshJobs;
diff --git a/Engine/Source/Developer/MeshReductionInterface/Private/MeshReductionManagerModule.cpp b/Engine/Source/Developer/MeshReductionInterface/Private/MeshReductionManagerModule.cpp
--- a/Engine/Source/Developer/MeshReductionInterface/Private/MeshReductionManagerModule.cpp
+++ b/Engine/Source/Developer/MeshReductionInterface/Private/MeshReductionManagerModule.cpp
@@ -187,6 +187,10 @@ void FMeshReductionManagerModule::StartupModule()
 	{
 		UE_LOG(LogMeshReduction, Display, TEXT("Using %s for distributed automatic mesh merging"), *DistributedMeshMergingModuleName.ToString());
 	}
+
+	bConfigurationCompleted = true;
+	ConfigurationCompletedDelegate.Broadcast();
+	ConfigurationCompletedDelegate.Clear();
 }
 
 void FMeshReductionManagerModule::ShutdownModule()
@@ -214,3 +218,16 @@ IMeshMerging* FMeshReductionManagerModule::GetDistributedMeshMergingInterface()
 {
 	return DistributedMeshMerging;
 }
+
+void FMeshReductionManagerModule::OnModuleConfigurationCompleted(FSimpleDelegate Delegate)
+{
+	// If the configuration delegate was already called, execute the provided delegate right away
+	if (bConfigurationCompleted)
+	{
+		Delegate.ExecuteIfBound();
+	}
+	else
+	{
+		ConfigurationCompletedDelegate.Add(Delegate);
+	}
+}
\ No newline at end of file
diff --git a/Engine/Source/Developer/MeshReductionInterface/Private/MeshReductionManagerModule.h b/Engine/Source/Developer/MeshReductionInterface/Private/MeshReductionManagerModule.h
--- a/Engine/Source/Developer/MeshReductionInterface/Private/MeshReductionManagerModule.h
+++ b/Engine/Source/Developer/MeshReductionInterface/Private/MeshReductionManagerModule.h
@@ -25,6 +25,9 @@ public:
 	virtual IMeshMerging* GetMeshMergingInterface() const override;
 
 	virtual IMeshMerging* GetDistributedMeshMergingInterface() const override;
+
+	virtual void OnModuleConfigurationCompleted(FSimpleDelegate Delegate) override;
+
 private:
 	/** Cached pointer to the mesh reduction interface. */
 	IMeshReduction* StaticMeshReduction;
@@ -34,4 +37,7 @@ private:
 	IMeshMerging* MeshMerging;
 	/** Cached pointer to the distributed mesh merging interface. */
 	IMeshMerging* DistributedMeshMerging;
+
+	bool bConfigurationCompleted = false;
+	FSimpleMulticastDelegate ConfigurationCompletedDelegate;
 };
diff --git a/Engine/Source/Developer/MeshReductionInterface/Public/IMeshReductionManagerModule.h b/Engine/Source/Developer/MeshReductionInterface/Public/IMeshReductionManagerModule.h
--- a/Engine/Source/Developer/MeshReductionInterface/Public/IMeshReductionManagerModule.h
+++ b/Engine/Source/Developer/MeshReductionInterface/Public/IMeshReductionManagerModule.h
@@ -3,6 +3,7 @@
 #pragma once
 
 #include "Modules/ModuleInterface.h"
+#include "Delegates/Delegate.h"
 
 class IMeshReduction;
 class IMeshMerging;
@@ -21,4 +22,7 @@ public:
 
 	/** Returns the distributed mesh merging plugin if available. */
 	virtual IMeshMerging* GetDistributedMeshMergingInterface() const = 0;
+
+	/** Register a delegate that will fire when the module configuration is completed. */
+	virtual void OnModuleConfigurationCompleted(FSimpleDelegate Delegate) = 0;
 };
