Expose two methods (ExportStaticMeshLOD and PropagateSplineDeformationToMesh) from MeshMergeHelpers to allow them to be used by other systems
This will be used to add Approximate Mesh HLOD support for Spline Meshes

diff --git a/Engine/Source/Developer/MeshMergeUtilities/Private/MeshMergeHelpers.cpp b/Engine/Source/Developer/MeshMergeUtilities/Private/MeshMergeHelpers.cpp
--- a/Engine/Source/Developer/MeshMergeUtilities/Private/MeshMergeHelpers.cpp
+++ b/Engine/Source/Developer/MeshMergeUtilities/Private/MeshMergeHelpers.cpp
@@ -239,7 +239,7 @@ static void RetrieveMeshInternal(const UStaticMesh* StaticMesh, const UStaticMes
 	check(StaticMesh);
 
 	// Export the mesh data using static mesh render data
-	FMeshMergeHelpers::ExportStaticMeshLOD(StaticMesh->GetRenderData()->LODResources[LODIndex], OutMeshDescription, StaticMesh->GetStaticMaterials());
+	StaticMesh->ExportStaticMeshLOD(StaticMesh->GetRenderData()->LODResources[LODIndex], OutMeshDescription);
 
 	// Make sure the mesh is not irreparably malformed.
 	if (OutMeshDescription.VertexInstances().Num() <= 0)
@@ -256,7 +256,7 @@ static void RetrieveMeshInternal(const UStaticMesh* StaticMesh, const UStaticMes
 		{
 			const USplineMeshComponent* SplineMeshComponent = Cast<USplineMeshComponent>(StaticMeshComponent);
 			// Deform mesh data according to the Spline Mesh Component's data
-			FMeshMergeHelpers::PropagateSplineDeformationToMesh(SplineMeshComponent, OutMeshDescription);
+			SplineMeshComponent->PropagateSplineDeformationToMesh(OutMeshDescription);
 		}
 
 		// If specified propagate painted vertex colors into our raw mesh
@@ -311,129 +311,6 @@ void FMeshMergeHelpers::RetrieveMesh(const UStaticMesh* StaticMesh, int32 LODInd
 	RetrieveMeshInternal(StaticMesh, nullptr, LODIndex, OutMeshDescription, /*bPropagateVertexColours*/false, /*bApplyComponentTransform*/false);
 }
 
-void FMeshMergeHelpers::ExportStaticMeshLOD(const FStaticMeshLODResources& StaticMeshLOD, FMeshDescription& OutMeshDescription, const TArray<FStaticMaterial>& Materials)
-{
-	const int32 NumWedges = StaticMeshLOD.IndexBuffer.GetNumIndices();
-	const int32 NumVertexPositions = StaticMeshLOD.VertexBuffers.PositionVertexBuffer.GetNumVertices();
-	const int32 NumFaces = NumWedges / 3;
-
-	OutMeshDescription.Empty();
-
-	if (NumVertexPositions <= 0 || StaticMeshLOD.VertexBuffers.StaticMeshVertexBuffer.GetNumVertices() <= 0)
-	{
-		return;
-	}
-
-	FStaticMeshAttributes Attributes(OutMeshDescription);
-	Attributes.Register();
-
-	TVertexAttributesRef<FVector3f> VertexPositions = Attributes.GetVertexPositions();
-	TEdgeAttributesRef<bool> EdgeHardnesses = Attributes.GetEdgeHardnesses();
-	TPolygonGroupAttributesRef<FName> PolygonGroupImportedMaterialSlotNames = Attributes.GetPolygonGroupMaterialSlotNames();
-	TVertexInstanceAttributesRef<FVector3f> VertexInstanceNormals = Attributes.GetVertexInstanceNormals();
-	TVertexInstanceAttributesRef<FVector3f> VertexInstanceTangents = Attributes.GetVertexInstanceTangents();
-	TVertexInstanceAttributesRef<float> VertexInstanceBinormalSigns = Attributes.GetVertexInstanceBinormalSigns();
-	TVertexInstanceAttributesRef<FVector4f> VertexInstanceColors = Attributes.GetVertexInstanceColors();
-	TVertexInstanceAttributesRef<FVector2f> VertexInstanceUVs = Attributes.GetVertexInstanceUVs();
-
-	OutMeshDescription.ReserveNewVertices(NumVertexPositions);
-	OutMeshDescription.ReserveNewVertexInstances(NumWedges);
-	OutMeshDescription.ReserveNewPolygons(NumFaces);
-	OutMeshDescription.ReserveNewEdges(NumWedges);
-
-	const int32 NumTexCoords = StaticMeshLOD.VertexBuffers.StaticMeshVertexBuffer.GetNumTexCoords();
-	VertexInstanceUVs.SetNumChannels(NumTexCoords);
-
-	TMap<int32, FPolygonGroupID> SectionToPolygonGroupMap;
-
-	for (int32 SectionIndex = 0; SectionIndex < StaticMeshLOD.Sections.Num(); ++SectionIndex)
-	{
-		const FStaticMeshSection& Section = StaticMeshLOD.Sections[SectionIndex];
-
-		// Skip empty sections
-		if (Section.NumTriangles == 0)
-		{
-			continue;
-		}
-
-		FPolygonGroupID CurrentPolygonGroupID = OutMeshDescription.CreatePolygonGroup();
-		SectionToPolygonGroupMap.Emplace(SectionIndex, CurrentPolygonGroupID);
-
-		if (Materials.IsValidIndex(Section.MaterialIndex))
-		{
-			PolygonGroupImportedMaterialSlotNames[CurrentPolygonGroupID] = Materials[Section.MaterialIndex].ImportedMaterialSlotName;
-		}
-		else
-		{
-			PolygonGroupImportedMaterialSlotNames[CurrentPolygonGroupID] = FName(*(TEXT("MeshMergeMaterial_") + FString::FromInt(SectionIndex)));
-		}
-	}
-
-	//Create the vertex
-	for (int32 VertexIndex = 0; VertexIndex < NumVertexPositions; ++VertexIndex)
-	{
-		FVertexID VertexID = OutMeshDescription.CreateVertex();
-		VertexPositions[VertexID] = StaticMeshLOD.VertexBuffers.PositionVertexBuffer.VertexPosition(VertexIndex);
-	}
-
-	//Create the vertex instances
-	for (int32 TriangleIndex = 0; TriangleIndex < NumFaces; ++TriangleIndex)
-	{
-		FPolygonGroupID CurrentPolygonGroupID = INDEX_NONE;
-		for (int32 SectionIndex = 0; SectionIndex < StaticMeshLOD.Sections.Num(); ++SectionIndex)
-		{
-			const FStaticMeshSection& Section = StaticMeshLOD.Sections[SectionIndex];
-			uint32 BeginTriangle = Section.FirstIndex / 3;
-			uint32 EndTriangle = BeginTriangle + Section.NumTriangles;
-			if ((uint32)TriangleIndex >= BeginTriangle && (uint32)TriangleIndex < EndTriangle)
-			{
-				CurrentPolygonGroupID = SectionToPolygonGroupMap[SectionIndex];
-				break;
-			}
-		}
-		check(CurrentPolygonGroupID != INDEX_NONE);
-
-		FVertexID VertexIDs[3];
-		TArray<FVertexInstanceID> VertexInstanceIDs;
-		VertexInstanceIDs.SetNum(3);
-
-		for (int32 Corner = 0; Corner < 3; ++Corner)
-		{
-			int32 WedgeIndex = StaticMeshLOD.IndexBuffer.GetIndex(TriangleIndex * 3 + Corner);
-			FVertexID VertexID(WedgeIndex);
-			FVertexInstanceID VertexInstanceID = OutMeshDescription.CreateVertexInstance(VertexID);
-			VertexIDs[Corner] = VertexID;
-			VertexInstanceIDs[Corner] = VertexInstanceID;
-
-			//NTBs
-			FVector TangentX = FVector4(StaticMeshLOD.VertexBuffers.StaticMeshVertexBuffer.VertexTangentX(WedgeIndex));
-			FVector TangentY = FVector(StaticMeshLOD.VertexBuffers.StaticMeshVertexBuffer.VertexTangentY(WedgeIndex));
-			FVector TangentZ = FVector4(StaticMeshLOD.VertexBuffers.StaticMeshVertexBuffer.VertexTangentZ(WedgeIndex));
-			VertexInstanceTangents[VertexInstanceID] = (FVector3f)TangentX;
-			VertexInstanceBinormalSigns[VertexInstanceID] = GetBasisDeterminantSign(TangentX, TangentY, TangentZ);
-			VertexInstanceNormals[VertexInstanceID] = (FVector3f)TangentZ;
-
-			// Vertex colors
-			if (StaticMeshLOD.VertexBuffers.ColorVertexBuffer.GetNumVertices() > 0)
-			{
-				VertexInstanceColors[VertexInstanceID] = FLinearColor(StaticMeshLOD.VertexBuffers.ColorVertexBuffer.VertexColor(WedgeIndex));
-			}
-			else
-			{
-				VertexInstanceColors[VertexInstanceID] = FLinearColor::White;
-			}
-
-			//Tex coord
-			for (int32 TexCoodIdx = 0; TexCoodIdx < NumTexCoords; ++TexCoodIdx)
-			{
-				VertexInstanceUVs.Set(VertexInstanceID, TexCoodIdx, StaticMeshLOD.VertexBuffers.StaticMeshVertexBuffer.GetVertexUV(WedgeIndex, TexCoodIdx));
-			}
-		}
-		//Create a polygon from this triangle
-		const FPolygonID NewPolygonID = OutMeshDescription.CreatePolygon(CurrentPolygonGroupID, VertexInstanceIDs);
-	}
-}
-
 void FMeshMergeHelpers::RetrieveMesh(const USkeletalMeshComponent* SkeletalMeshComponent, int32 LODIndex, FMeshDescription& OutMeshDescription, bool bPropagateVertexColours, bool bApplyComponentTransform)
 {
 	FSkeletalMeshModel* Resource = SkeletalMeshComponent->GetSkeletalMeshAsset()->GetImportedModel();
@@ -741,43 +618,6 @@ void FMeshMergeHelpers::CullTrianglesFromVolumesAndUnderLandscapes(const UWorld*
 	}
 }
 
-void FMeshMergeHelpers::PropagateSplineDeformationToMesh(const USplineMeshComponent* InSplineMeshComponent, FMeshDescription& InOutMeshDescription)
-{
-	FStaticMeshAttributes Attributes(InOutMeshDescription);
-
-	TVertexAttributesRef<FVector3f> VertexPositions = Attributes.GetVertexPositions();
-	TVertexInstanceAttributesRef<FVector3f> VertexInstanceNormals = Attributes.GetVertexInstanceNormals();
-	TVertexInstanceAttributesRef<FVector3f> VertexInstanceTangents = Attributes.GetVertexInstanceTangents();
-	TVertexInstanceAttributesRef<float> VertexInstanceBinormalSigns = Attributes.GetVertexInstanceBinormalSigns();
-
-	// Apply spline deformation for each vertex's tangents
-	int32 WedgeIndex = 0;
-	for (const FTriangleID TriangleID : InOutMeshDescription.Triangles().GetElementIDs())
-	{
-		for (int32 Corner = 0; Corner < 3; ++Corner, ++WedgeIndex)
-		{
-			const FVertexInstanceID VertexInstanceID = InOutMeshDescription.GetTriangleVertexInstance(TriangleID, Corner);
-			const FVertexID VertexID = InOutMeshDescription.GetVertexInstanceVertex(VertexInstanceID);
-			const float& AxisValue = USplineMeshComponent::GetAxisValueRef(VertexPositions[VertexID], InSplineMeshComponent->ForwardAxis);
-			FTransform SliceTransform = InSplineMeshComponent->CalcSliceTransform(AxisValue);
-			FVector TangentY = FVector::CrossProduct((FVector)VertexInstanceNormals[VertexInstanceID], (FVector)VertexInstanceTangents[VertexInstanceID]).GetSafeNormal() * VertexInstanceBinormalSigns[VertexInstanceID];
-			VertexInstanceTangents[VertexInstanceID] = (FVector3f)SliceTransform.TransformVector((FVector)VertexInstanceTangents[VertexInstanceID]);
-			TangentY = SliceTransform.TransformVector(TangentY);
-			VertexInstanceNormals[VertexInstanceID] = (FVector3f)SliceTransform.TransformVector((FVector)VertexInstanceNormals[VertexInstanceID]);
-			VertexInstanceBinormalSigns[VertexInstanceID] = GetBasisDeterminantSign((FVector)VertexInstanceTangents[VertexInstanceID], TangentY, (FVector)VertexInstanceNormals[VertexInstanceID]);
-		}
-	}
-
-	// Apply spline deformation for each vertex position
-	for (const FVertexID VertexID : InOutMeshDescription.Vertices().GetElementIDs())
-	{
-		float& AxisValue = USplineMeshComponent::GetAxisValueRef(VertexPositions[VertexID], InSplineMeshComponent->ForwardAxis);
-		FTransform SliceTransform = InSplineMeshComponent->CalcSliceTransform(AxisValue);
-		AxisValue = 0.0f;
-		VertexPositions[VertexID] = (FVector3f)SliceTransform.TransformPosition((FVector)VertexPositions[VertexID]);
-	}
-}
-
 void FMeshMergeHelpers::PropagateSplineDeformationToPhysicsGeometry(USplineMeshComponent* SplineMeshComponent, FKAggregateGeom& InOutPhysicsGeometry)
 {
 	const FVector Mask = USplineMeshComponent::GetAxisMask(SplineMeshComponent->GetForwardAxis());
diff --git a/Engine/Source/Developer/MeshMergeUtilities/Private/MeshMergeHelpers.h b/Engine/Source/Developer/MeshMergeUtilities/Private/MeshMergeHelpers.h
--- a/Engine/Source/Developer/MeshMergeUtilities/Private/MeshMergeHelpers.h
+++ b/Engine/Source/Developer/MeshMergeUtilities/Private/MeshMergeHelpers.h
@@ -37,9 +37,6 @@ public:
 	static UE_API void RetrieveMesh(const UStaticMeshComponent* StaticMeshComponent, int32 LODIndex, FMeshDescription& OutMeshDescription, bool bPropagateVertexColours, bool bApplyComponentTransform = true);
 	static UE_API void RetrieveMesh(const USkeletalMeshComponent* SkeletalMeshComponent, int32 LODIndex, FMeshDescription& OutMeshDescription, bool bPropagateVertexColours, bool bApplyComponentTransform = false);
 	static UE_API void RetrieveMesh(const UStaticMesh* StaticMesh, int32 LODIndex, FMeshDescription& OutMeshDescription);
-	
-	/** Exports static mesh LOD render data to a RawMesh */
-	static UE_API void ExportStaticMeshLOD(const FStaticMeshLODResources& StaticMeshLOD, FMeshDescription& OutMeshDescription, const TArray<FStaticMaterial>& Materials);
 
 	/** Checks whether or not the texture coordinates are outside of 0-1 UV ranges */
 	static UE_API bool CheckWrappingUVs(const TArray<FVector2D>& UVs);
@@ -48,9 +45,6 @@ public:
 	/** Culls away triangles which are inside culling volumes or completely underneath the landscape */
 	static UE_API void CullTrianglesFromVolumesAndUnderLandscapes(const UWorld* World, const FBoxSphereBounds& Bounds, FMeshDescription& InOutRawMesh);
 	
-	/** Propagates deformation along spline to mesh data */
-	static UE_API void PropagateSplineDeformationToMesh(const USplineMeshComponent* InSplineMeshComponent, FMeshDescription& InOutMeshDescription);
-	
 	/** Propagates deformation along spline to physics geometry data */
 	static UE_API void PropagateSplineDeformationToPhysicsGeometry(USplineMeshComponent* SplineMeshComponent, FKAggregateGeom& InOutPhysicsGeometry);
 
diff --git a/Engine/Source/Runtime/Engine/Classes/Components/SplineMeshComponent.h b/Engine/Source/Runtime/Engine/Classes/Components/SplineMeshComponent.h
--- a/Engine/Source/Runtime/Engine/Classes/Components/SplineMeshComponent.h
+++ b/Engine/Source/Runtime/Engine/Classes/Components/SplineMeshComponent.h
@@ -15,6 +15,7 @@ class FPrimitiveSceneProxy;
 class ULightComponent;
 struct FNavigableGeometryExport;
 class UBodySetup;
+struct FMeshDescription;
 
 // Helper for packing spline mesh shader parameters into a float4 buffer
 ENGINE_API void PackSplineMeshParams(const FSplineMeshShaderParams& Params, const TArrayView<FVector4f>& Output);
@@ -422,6 +423,9 @@ public:
 	 */
 	ENGINE_API FTransform CalcSliceTransformAtSplineOffset(const float Alpha, const float MinT=0.f, const float MaxT=1.0f) const;
 
+	/** Propagates deformation along spline to mesh data */
+	ENGINE_API void PropagateSplineDeformationToMesh(FMeshDescription& InOutMeshDescription) const;
+
 	inline static const double& GetAxisValueRef(const FVector3d& InVector, ESplineMeshAxis::Type InAxis);
 	inline static double& GetAxisValueRef(FVector3d& InVector, ESplineMeshAxis::Type InAxis);
 
diff --git a/Engine/Source/Runtime/Engine/Classes/Engine/StaticMesh.h b/Engine/Source/Runtime/Engine/Classes/Engine/StaticMesh.h
--- a/Engine/Source/Runtime/Engine/Classes/Engine/StaticMesh.h
+++ b/Engine/Source/Runtime/Engine/Classes/Engine/StaticMesh.h
@@ -1874,6 +1874,9 @@ public:
 	/** Builds a LOD resource from a MeshDescription */
 	ENGINE_API void BuildFromMeshDescription(const FMeshDescription& MeshDescription, FStaticMeshLODResources& LODResources);
 
+	/** Exports static mesh LOD render data to a MeshDescription */
+	ENGINE_API void ExportStaticMeshLOD(const FStaticMeshLODResources& InLODResources, FMeshDescription& OutMeshDescription) const;
+
 	/**
 	 * Returns the number of UV channels for the given LOD of a StaticMesh.
 	 * @param	LODIndex			Index of the StaticMesh LOD.
diff --git a/Engine/Source/Runtime/Engine/Private/Components/SplineMeshComponent.cpp b/Engine/Source/Runtime/Engine/Private/Components/SplineMeshComponent.cpp
--- a/Engine/Source/Runtime/Engine/Private/Components/SplineMeshComponent.cpp
+++ b/Engine/Source/Runtime/Engine/Private/Components/SplineMeshComponent.cpp
@@ -25,6 +25,7 @@
 #include "ComponentRecreateRenderStateContext.h"
 #include "Engine/World.h"
 #include "NaniteVertexFactory.h"
+#include "StaticMeshAttributes.h"
 #include "StaticMeshComponentLODInfo.h"
 #include "SplineMeshSceneProxyDesc.h"
 
@@ -813,6 +814,43 @@ FTransform USplineMeshComponent::CalcSliceTransformAtSplineOffset(const float Al
 	return FSplineMeshSceneProxyDesc(this).CalcSliceTransformAtSplineOffset(Alpha, MinT, MaxT);
 }
 
+void USplineMeshComponent::PropagateSplineDeformationToMesh(FMeshDescription& InOutMeshDescription) const
+{
+	FStaticMeshAttributes Attributes(InOutMeshDescription);
+
+	TVertexAttributesRef<FVector3f> VertexPositions = Attributes.GetVertexPositions();
+	TVertexInstanceAttributesRef<FVector3f> VertexInstanceNormals = Attributes.GetVertexInstanceNormals();
+	TVertexInstanceAttributesRef<FVector3f> VertexInstanceTangents = Attributes.GetVertexInstanceTangents();
+	TVertexInstanceAttributesRef<float> VertexInstanceBinormalSigns = Attributes.GetVertexInstanceBinormalSigns();
+
+	// Apply spline deformation for each vertex's tangents
+	int32 WedgeIndex = 0;
+	for (const FTriangleID TriangleID : InOutMeshDescription.Triangles().GetElementIDs())
+	{
+		for (int32 Corner = 0; Corner < 3; ++Corner, ++WedgeIndex)
+		{
+			const FVertexInstanceID VertexInstanceID = InOutMeshDescription.GetTriangleVertexInstance(TriangleID, Corner);
+			const FVertexID VertexID = InOutMeshDescription.GetVertexInstanceVertex(VertexInstanceID);
+			const float& AxisValue = USplineMeshComponent::GetAxisValueRef(VertexPositions[VertexID], ForwardAxis);
+			FTransform SliceTransform = CalcSliceTransform(AxisValue);
+			FVector TangentY = FVector::CrossProduct((FVector)VertexInstanceNormals[VertexInstanceID], (FVector)VertexInstanceTangents[VertexInstanceID]).GetSafeNormal() * VertexInstanceBinormalSigns[VertexInstanceID];
+			VertexInstanceTangents[VertexInstanceID] = (FVector3f)SliceTransform.TransformVector((FVector)VertexInstanceTangents[VertexInstanceID]);
+			TangentY = SliceTransform.TransformVector(TangentY);
+			VertexInstanceNormals[VertexInstanceID] = (FVector3f)SliceTransform.TransformVector((FVector)VertexInstanceNormals[VertexInstanceID]);
+			VertexInstanceBinormalSigns[VertexInstanceID] = GetBasisDeterminantSign((FVector)VertexInstanceTangents[VertexInstanceID], TangentY, (FVector)VertexInstanceNormals[VertexInstanceID]);
+		}
+	}
+
+	// Apply spline deformation for each vertex position
+	for (const FVertexID VertexID : InOutMeshDescription.Vertices().GetElementIDs())
+	{
+		float& AxisValue = USplineMeshComponent::GetAxisValueRef(VertexPositions[VertexID], ForwardAxis);
+		FTransform SliceTransform = CalcSliceTransform(AxisValue);
+		AxisValue = 0.0f;
+		VertexPositions[VertexID] = (FVector3f)SliceTransform.TransformPosition((FVector)VertexPositions[VertexID]);
+	}
+}
+
 bool USplineMeshComponent::GetPhysicsTriMeshData(struct FTriMeshCollisionData* CollisionData, bool InUseAllTriData)
 {
 	if (GetStaticMesh())
diff --git a/Engine/Source/Runtime/Engine/Private/StaticMesh.cpp b/Engine/Source/Runtime/Engine/Private/StaticMesh.cpp
--- a/Engine/Source/Runtime/Engine/Private/StaticMesh.cpp
+++ b/Engine/Source/Runtime/Engine/Private/StaticMesh.cpp
@@ -8543,6 +8543,129 @@ void UStaticMesh::BuildFromMeshDescription(const FMeshDescription& MeshDescripti
 	LODResources.AdditionalIndexBuffers->ReversedDepthOnlyIndexBuffer.SetIndices(ReversedIndexBuffer, IndexBufferStride);
 }
 
+void UStaticMesh::ExportStaticMeshLOD(const FStaticMeshLODResources& InLODResources, FMeshDescription& OutMeshDescription) const
+{
+	const int32 NumWedges = InLODResources.IndexBuffer.GetNumIndices();
+	const int32 NumVertexPositions = InLODResources.VertexBuffers.PositionVertexBuffer.GetNumVertices();
+	const int32 NumFaces = NumWedges / 3;
+
+	OutMeshDescription.Empty();
+
+	if (NumVertexPositions <= 0 || InLODResources.VertexBuffers.StaticMeshVertexBuffer.GetNumVertices() <= 0)
+	{
+		return;
+	}
+
+	FStaticMeshAttributes Attributes(OutMeshDescription);
+	Attributes.Register();
+
+	TVertexAttributesRef<FVector3f> VertexPositions = Attributes.GetVertexPositions();
+	TEdgeAttributesRef<bool> EdgeHardnesses = Attributes.GetEdgeHardnesses();
+	TPolygonGroupAttributesRef<FName> PolygonGroupImportedMaterialSlotNames = Attributes.GetPolygonGroupMaterialSlotNames();
+	TVertexInstanceAttributesRef<FVector3f> VertexInstanceNormals = Attributes.GetVertexInstanceNormals();
+	TVertexInstanceAttributesRef<FVector3f> VertexInstanceTangents = Attributes.GetVertexInstanceTangents();
+	TVertexInstanceAttributesRef<float> VertexInstanceBinormalSigns = Attributes.GetVertexInstanceBinormalSigns();
+	TVertexInstanceAttributesRef<FVector4f> VertexInstanceColors = Attributes.GetVertexInstanceColors();
+	TVertexInstanceAttributesRef<FVector2f> VertexInstanceUVs = Attributes.GetVertexInstanceUVs();
+
+	OutMeshDescription.ReserveNewVertices(NumVertexPositions);
+	OutMeshDescription.ReserveNewVertexInstances(NumWedges);
+	OutMeshDescription.ReserveNewPolygons(NumFaces);
+	OutMeshDescription.ReserveNewEdges(NumWedges);
+
+	const int32 NumTexCoords = InLODResources.VertexBuffers.StaticMeshVertexBuffer.GetNumTexCoords();
+	VertexInstanceUVs.SetNumChannels(NumTexCoords);
+
+	TMap<int32, FPolygonGroupID> SectionToPolygonGroupMap;
+
+	for (int32 SectionIndex = 0; SectionIndex < InLODResources.Sections.Num(); ++SectionIndex)
+	{
+		const FStaticMeshSection& Section = InLODResources.Sections[SectionIndex];
+
+		// Skip empty sections
+		if (Section.NumTriangles == 0)
+		{
+			continue;
+		}
+
+		FPolygonGroupID CurrentPolygonGroupID = OutMeshDescription.CreatePolygonGroup();
+		SectionToPolygonGroupMap.Emplace(SectionIndex, CurrentPolygonGroupID);
+
+		if (GetStaticMaterials().IsValidIndex(Section.MaterialIndex))
+		{
+			PolygonGroupImportedMaterialSlotNames[CurrentPolygonGroupID] = GetStaticMaterials()[Section.MaterialIndex].ImportedMaterialSlotName;
+		}
+		else
+		{
+			PolygonGroupImportedMaterialSlotNames[CurrentPolygonGroupID] = FName(*(TEXT("MeshMergeMaterial_") + FString::FromInt(SectionIndex)));
+		}
+	}
+
+	//Create the vertex
+	for (int32 VertexIndex = 0; VertexIndex < NumVertexPositions; ++VertexIndex)
+	{
+		FVertexID VertexID = OutMeshDescription.CreateVertex();
+		VertexPositions[VertexID] = InLODResources.VertexBuffers.PositionVertexBuffer.VertexPosition(VertexIndex);
+	}
+
+	//Create the vertex instances
+	for (int32 TriangleIndex = 0; TriangleIndex < NumFaces; ++TriangleIndex)
+	{
+		FPolygonGroupID CurrentPolygonGroupID = INDEX_NONE;
+		for (int32 SectionIndex = 0; SectionIndex < InLODResources.Sections.Num(); ++SectionIndex)
+		{
+			const FStaticMeshSection& Section = InLODResources.Sections[SectionIndex];
+			uint32 BeginTriangle = Section.FirstIndex / 3;
+			uint32 EndTriangle = BeginTriangle + Section.NumTriangles;
+			if ((uint32)TriangleIndex >= BeginTriangle && (uint32)TriangleIndex < EndTriangle)
+			{
+				CurrentPolygonGroupID = SectionToPolygonGroupMap[SectionIndex];
+				break;
+			}
+		}
+		check(CurrentPolygonGroupID != INDEX_NONE);
+
+		FVertexID VertexIDs[3];
+		TArray<FVertexInstanceID> VertexInstanceIDs;
+		VertexInstanceIDs.SetNum(3);
+
+		for (int32 Corner = 0; Corner < 3; ++Corner)
+		{
+			int32 WedgeIndex = InLODResources.IndexBuffer.GetIndex(TriangleIndex * 3 + Corner);
+			FVertexID VertexID(WedgeIndex);
+			FVertexInstanceID VertexInstanceID = OutMeshDescription.CreateVertexInstance(VertexID);
+			VertexIDs[Corner] = VertexID;
+			VertexInstanceIDs[Corner] = VertexInstanceID;
+
+			//NTBs
+			FVector TangentX = FVector4(InLODResources.VertexBuffers.StaticMeshVertexBuffer.VertexTangentX(WedgeIndex));
+			FVector TangentY = FVector(InLODResources.VertexBuffers.StaticMeshVertexBuffer.VertexTangentY(WedgeIndex));
+			FVector TangentZ = FVector4(InLODResources.VertexBuffers.StaticMeshVertexBuffer.VertexTangentZ(WedgeIndex));
+			VertexInstanceTangents[VertexInstanceID] = (FVector3f)TangentX;
+			VertexInstanceBinormalSigns[VertexInstanceID] = GetBasisDeterminantSign(TangentX, TangentY, TangentZ);
+			VertexInstanceNormals[VertexInstanceID] = (FVector3f)TangentZ;
+
+			// Vertex colors
+			if (InLODResources.VertexBuffers.ColorVertexBuffer.GetNumVertices() > 0)
+			{
+				VertexInstanceColors[VertexInstanceID] = FLinearColor(InLODResources.VertexBuffers.ColorVertexBuffer.VertexColor(WedgeIndex));
+			}
+			else
+			{
+				VertexInstanceColors[VertexInstanceID] = FLinearColor::White;
+			}
+
+			//Tex coord
+			for (int32 TexCoodIdx = 0; TexCoodIdx < NumTexCoords; ++TexCoodIdx)
+			{
+				VertexInstanceUVs.Set(VertexInstanceID, TexCoodIdx, InLODResources.VertexBuffers.StaticMeshVertexBuffer.GetVertexUV(WedgeIndex, TexCoodIdx));
+			}
+		}
+		//Create a polygon from this triangle
+		const FPolygonID NewPolygonID = OutMeshDescription.CreatePolygon(CurrentPolygonGroupID, VertexInstanceIDs);
+	}
+}
+
 
 UStaticMeshDescription* UStaticMesh::CreateStaticMeshDescription(UObject* Outer)
 {
