Mesh Merge - Duplicate non-assets MICs if necessary, as it's already done for MIDs

#DisabledValidations Preflight

diff --git a/Engine/Source/Developer/MeshMergeUtilities/Private/MeshMergeUtilities.cpp b/Engine/Source/Developer/MeshMergeUtilities/Private/MeshMergeUtilities.cpp
--- a/Engine/Source/Developer/MeshMergeUtilities/Private/MeshMergeUtilities.cpp
+++ b/Engine/Source/Developer/MeshMergeUtilities/Private/MeshMergeUtilities.cpp
@@ -25,6 +25,7 @@
 #include "UObject/UObjectBaseUtility.h"
 #include "UObject/Package.h"
 #include "Materials/Material.h"
+#include "Materials/MaterialInstanceConstant.h"
 #include "Materials/MaterialInstanceDynamic.h"
 #include "Misc/ScopedSlowTask.h"
 #include "Modules/ModuleManager.h"
@@ -2640,10 +2641,11 @@ void FMeshMergeUtilities::MergeComponentsToStaticMesh(const TArray<UPrimitiveCom
 		{
 			if (Material && (!Material->IsAsset() && InOuter != GetTransientPackage()))
 			{
-				// MIDs are not assets, duplicate them and outer them to the static mesh.
-				if (UMaterialInstanceDynamic* MID = Cast<UMaterialInstanceDynamic>(Material))
+				// Only allow duplication for MIDs and MICs
+				if (Material->IsA<UMaterialInstanceDynamic>() || Material->IsA<UMaterialInstanceConstant>())
 				{
-					Material = DuplicateObject<UMaterialInstanceDynamic>(MID, StaticMesh);
+					// Duplicate them and outer them to the static mesh.
+					Material = DuplicateObject<UMaterialInterface>(Material, StaticMesh);
 				}
 				else
 				{
