[Substrate] Fixed most material properties not being added to FMaterialData::PropertySizes when baking Substrate materials on the Actor Merge -> Simplify path.

diff --git a/Engine/Source/Developer/MeshMergeUtilities/Private/MeshMergeUtilities.cpp b/Engine/Source/Developer/MeshMergeUtilities/Private/MeshMergeUtilities.cpp
--- a/Engine/Source/Developer/MeshMergeUtilities/Private/MeshMergeUtilities.cpp
+++ b/Engine/Source/Developer/MeshMergeUtilities/Private/MeshMergeUtilities.cpp
@@ -1738,7 +1738,11 @@ void FMeshMergeUtilities::CreateProxyMesh(const TArray<UStaticMeshComponent*>& I
 
 			for (const FPropertyEntry& Entry : Options->Properties)
 			{
-				if (!Entry.bUseConstantValue && Material->IsPropertyActive(Entry.Property) && Entry.Property != MP_MAX)
+				// With substrate, we cannot use connection on the root node to infer Active or Constant properties. So we activate all material input for now. 
+				// SUBSTRATE_TODO: gather information about the material used input, maybe using Material->GetSubstrateCompilationConfig().
+				const float bIsSubstrateEnabled = Substrate::IsSubstrateEnabled();
+
+				if ((bIsSubstrateEnabled || (!Entry.bUseConstantValue && Material->IsPropertyActive(Entry.Property))) && Entry.Property != MP_MAX)
 				{
 					MaterialSettings.PropertySizes.Add(Entry.Property, Entry.bUseCustomSize ? Entry.CustomSize : Options->TextureSize);
 				}
