Mesh Merge - Fix issue where materials relying on World Position would be baked improperly

When merging meshes with materials that read world position multiple instances of the same static mesh were producing identical baked results regardless of their world location.
In order to prevent these from being merged together, we now analyze the material to find any use of the world position (through bNeedsWorldPositionExcludingShaderOffsets). The deduplication optimization is preserved for materials that don't read world position.

diff --git a/Engine/Source/Developer/MeshMergeUtilities/Private/MeshMergeUtilities.cpp b/Engine/Source/Developer/MeshMergeUtilities/Private/MeshMergeUtilities.cpp
--- a/Engine/Source/Developer/MeshMergeUtilities/Private/MeshMergeUtilities.cpp
+++ b/Engine/Source/Developer/MeshMergeUtilities/Private/MeshMergeUtilities.cpp
@@ -547,27 +547,39 @@ void FMeshMergeUtilities::BakeMaterialsForMesh(UStaticMesh* StaticMesh) const
 }
 
 
-static bool DetermineMaterialVertexDataUsage(UMaterialInterface* Material, const UMaterialOptions* MaterialOptions)
+struct FMaterialDataUsage
 {
-	TRACE_CPUPROFILER_EVENT_SCOPE(DetermineMaterialVertexDataUsage);
+	bool bUsesVertexData = false;
+	bool bUsesWorldPosition = false;
+};
+
+static FMaterialDataUsage DetermineMaterialDataUsage(UMaterialInterface* Material, const UMaterialOptions* MaterialOptions)
+{
+	TRACE_CPUPROFILER_EVENT_SCOPE(DetermineMaterialDataUsage);
+
+	FMaterialDataUsage Usage;
 
 	for (const FPropertyEntry& Entry : MaterialOptions->Properties)
 	{
 		// Don't have to check a property if the result is going to be constant anyway
 		if (!Entry.bUseConstantValue && Entry.Property != MP_MAX)
 		{
-			int32 NumTextureCoordinates;
-			bool bUsesVertexData;
-			Material->AnalyzeMaterialProperty(Entry.Property, NumTextureCoordinates, bUsesVertexData);
+			FMaterialAnalysisResult Result;
+			Material->AnalyzeMaterialPropertyEx(Entry.Property, Result);
+
+			int32 NumTextureCoordinates = Result.TextureCoordinates.FindLast(true) + 1;
+
+			Usage.bUsesVertexData |= Result.bRequiresVertexData || NumTextureCoordinates > 1;
+			Usage.bUsesWorldPosition |= Result.bRequiresWorldPosition;
 
-			if (bUsesVertexData || NumTextureCoordinates > 1)
+			if (Usage.bUsesVertexData && Usage.bUsesWorldPosition)
 			{
-				return true;
+				break;
 			}
 		}
 	}
 
-	return false;
+	return Usage;
 }
 
 void FMeshMergeUtilities::ConvertOutputToFlatMaterials(const TArray<FBakeOutput>& BakeOutputs, const TArray<FMaterialData>& MaterialData, TArray<FFlattenMaterial> &FlattenedMaterials) const
@@ -2793,19 +2805,29 @@ void FMeshMergeUtilities::CreateMergedMaterial(FMeshMergeDataTracker& InDataTrac
 	UMaterialOptions* MaterialOptions = PopulateMaterialOptions(MaterialProxySettings);
 	TGCObjectScopeGuard<UMaterialOptions> MaterialOptionsGCScopeGuard(MaterialOptions);
 
-	// Check each material to see if the shader actually uses vertex data and collect flags
-	TArray<TOptional<bool>> bMaterialUsesVertexData;
-	bMaterialUsesVertexData.SetNum(InUniqueMaterials.Num());
+	// Check each material to see if the shader actually uses vertex data / world position.
+	// Deferred and cached per material, as AnalyzeMaterialPropertyEx() is a costly call.
+	TArray<TOptional<FMaterialDataUsage>> CachedMaterialDataUsage;
+	CachedMaterialDataUsage.SetNum(InUniqueMaterials.Num());
 
-	// Deferred call, as this may not be required by all code paths and is pretty costly to compute
-	auto DoesMaterialUsesVertexData = [&](const int32 InMaterialIndex)
+	auto GetMaterialDataUsage = [&](const int32 InMaterialIndex)
 	{
-		if (!bMaterialUsesVertexData[InMaterialIndex].IsSet())
+		if (!CachedMaterialDataUsage[InMaterialIndex].IsSet())
 		{
-			bMaterialUsesVertexData[InMaterialIndex] = DetermineMaterialVertexDataUsage(InUniqueMaterials[InMaterialIndex], MaterialOptions);
+			CachedMaterialDataUsage[InMaterialIndex] = DetermineMaterialDataUsage(InUniqueMaterials[InMaterialIndex], MaterialOptions);
 		}
 
-		return bMaterialUsesVertexData[InMaterialIndex].GetValue();
+		return CachedMaterialDataUsage[InMaterialIndex].GetValue();
+	};
+
+	auto DoesMaterialUsesVertexData = [&](const int32 InMaterialIndex)
+	{
+		return GetMaterialDataUsage(InMaterialIndex).bUsesVertexData;
+	};
+
+	auto DoesMaterialUseWorldPosition = [&](const int32 InMaterialIndex)
+	{
+		return GetMaterialDataUsage(InMaterialIndex).bUsesWorldPosition;
 	};
 
 	// For each unique material calculate how 'important' they are
@@ -2902,9 +2924,12 @@ void FMeshMergeUtilities::CreateMergedMaterial(FMeshMergeDataTracker& InDataTrac
 				return (!LHS && !RHS) ? true : (!LHS || !RHS) ? false : CompareCustomPrimitiveData(LHS->CustomPrimitiveData, RHS->CustomPrimitiveData);
 			};
 
-			auto CompareMeshData = [&ComparePrimitiveData](const FMeshData& LHS, const FMeshData& RHS)
+			// When the material reads world position, instances of the same mesh at different world positions must not share a bake entry, as they produce different results.
+			const bool bMaterialNeedsUniqueInstances = DoesMaterialUseWorldPosition(MaterialIndex);
+
+			auto CompareMeshData = [&ComparePrimitiveData, bMaterialNeedsUniqueInstances](const FMeshData& LHS, const FMeshData& RHS)
 			{
-				return (LHS.Mesh == RHS.Mesh) && (LHS.MaterialIndices == RHS.MaterialIndices) && (LHS.bMirrored == RHS.bMirrored) && (LHS.VertexColorHash == RHS.VertexColorHash) && ComparePrimitiveData(LHS.PrimitiveData, RHS.PrimitiveData);
+				return (LHS.Mesh == RHS.Mesh) && (!bMaterialNeedsUniqueInstances || LHS.MeshDescription == RHS.MeshDescription) && (LHS.MaterialIndices == RHS.MaterialIndices) && (LHS.bMirrored == RHS.bMirrored) && (LHS.VertexColorHash == RHS.VertexColorHash) && ComparePrimitiveData(LHS.PrimitiveData, RHS.PrimitiveData);
 			};
 
 			// Find material & mesh pair
@@ -3021,9 +3046,22 @@ void FMeshMergeUtilities::CreateMergedMaterial(FMeshMergeDataTracker& InDataTrac
 		CreateMergedRawMeshes(InDataTracker, RemapUVMergeSettings, InStaticMeshComponentsToMerge, InUniqueMaterials, InCollapsedMaterialMap, InOutputMaterialsMap, false, false, InMergedAssetPivot, MergedRawMeshes);
 
 		// Create texture coords for the merged mesh
+		// Only merge identical triangles if no material reads world position,
+		// as instances of the same mesh at different world positions would otherwise share
+		// UV space in the atlas and produce identical baked results.
+		bool bAnyMaterialUsesWorldPosition = false;
+		for (int32 MaterialIndex = 0; MaterialIndex < InUniqueMaterials.Num(); ++MaterialIndex)
+		{
+			if (DoesMaterialUseWorldPosition(MaterialIndex))
+			{
+				bAnyMaterialUsesWorldPosition = true;
+				break;
+			}
+		}
+
 		FStaticMeshOperations::FGenerateUVOptions GenerateUVOptions;
 		GenerateUVOptions.TextureResolution = MaterialOptions->TextureSize.GetMax();
-		GenerateUVOptions.bMergeTrianglesWithIdenticalAttributes = true;
+		GenerateUVOptions.bMergeTrianglesWithIdenticalAttributes = !bAnyMaterialUsesWorldPosition;
 		GenerateUVOptions.UVMethod = GetUVGenerationMethodToUse();
 
 		TArray<FVector2D> GlobalTextureCoordinates;
diff --git a/Engine/Source/Runtime/Engine/Private/Materials/MaterialShared.cpp b/Engine/Source/Runtime/Engine/Private/Materials/MaterialShared.cpp
--- a/Engine/Source/Runtime/Engine/Private/Materials/MaterialShared.cpp
+++ b/Engine/Source/Runtime/Engine/Private/Materials/MaterialShared.cpp
@@ -5264,6 +5264,7 @@ void UMaterialInterface::AnalyzeMaterialCompilationInCallback(TFunctionRef<void
 		MaterialTranslator.bUsesAOMaterialMask ||
 		MaterialTranslator.bUsesLightmapUVs ||
 		MaterialTranslator.bUsesVertexPosition;
+	OutResult.bRequiresWorldPosition = MaterialTranslator.bNeedsWorldPositionExcludingShaderOffsets;
 	OutResult.EstimatedNumTextureSamplesVS = TempOutput.EstimatedNumTextureSamplesVS;
 	OutResult.EstimatedNumTextureSamplesPS = TempOutput.EstimatedNumTextureSamplesPS;
 #endif
diff --git a/Engine/Source/Runtime/Engine/Public/Materials/MaterialInterface.h b/Engine/Source/Runtime/Engine/Public/Materials/MaterialInterface.h
--- a/Engine/Source/Runtime/Engine/Public/Materials/MaterialInterface.h
+++ b/Engine/Source/Runtime/Engine/Public/Materials/MaterialInterface.h
@@ -263,6 +263,9 @@ struct FMaterialAnalysisResult
 	/** Whether any vertex data is used */
 	bool bRequiresVertexData = false;
 
+	/** Whether the world position (excluding shader offsets) is used */
+	bool bRequiresWorldPosition = false;
+
 	/** If material translation was success, valid in AnalyzeMaterialTranslationOutput */
 	bool bTranslationSuccess = false;
 
